UdpMessagingPlugin: only printing a warning endpoint configuration is overriden and invalid

diff --git a/Engine/Plugins/Messaging/UdpMessaging/Source/UdpMessaging/Private/UdpMessagingModule.cpp b/Engine/Plugins/Messaging/UdpMessaging/Source/UdpMessaging/Private/UdpMessagingModule.cpp
--- a/Engine/Plugins/Messaging/UdpMessaging/Source/UdpMessaging/Private/UdpMessagingModule.cpp
+++ b/Engine/Plugins/Messaging/UdpMessaging/Source/UdpMessaging/Private/UdpMessagingModule.cpp
@@ -114,7 +114,11 @@ protected:
 
 		if (!FIPv4Endpoint::Parse(Settings->UnicastEndpoint, UnicastEndpoint))
 		{
-			GLog->Logf(TEXT("Warning: Invalid UDP Messaging UnicastEndpoint '%s' - binding to all local network adapters instead"), *Settings->UnicastEndpoint);
+			if (!Settings->UnicastEndpoint.IsEmpty())
+			{
+				GLog->Logf(TEXT("Warning: Invalid UDP Messaging UnicastEndpoint '%s' - binding to all local network adapters instead"), *Settings->UnicastEndpoint);
+			}
+
 			UnicastEndpoint = FIPv4Endpoint::Any;
 			Settings->UnicastEndpoint = UnicastEndpoint.ToText().ToString();
 			ResaveSettings = true;
@@ -122,7 +126,11 @@ protected:
 
 		if (!FIPv4Endpoint::Parse(Settings->MulticastEndpoint, MulticastEndpoint))
 		{
-			GLog->Logf(TEXT("Warning: Invalid UDP Messaging MulticastEndpoint '%s' - using default endpoint '%s' instead"), *Settings->MulticastEndpoint, *UDP_MESSAGING_DEFAULT_MULTICAST_ENDPOINT.ToText().ToString());
+			if (!Settings->MulticastEndpoint.IsEmpty())
+			{
+				GLog->Logf(TEXT("Warning: Invalid UDP Messaging MulticastEndpoint '%s' - using default endpoint '%s' instead"), *Settings->MulticastEndpoint, *UDP_MESSAGING_DEFAULT_MULTICAST_ENDPOINT.ToText().ToString());
+			}
+
 			MulticastEndpoint = UDP_MESSAGING_DEFAULT_MULTICAST_ENDPOINT;
 			Settings->MulticastEndpoint = MulticastEndpoint.ToText().ToString();
 			ResaveSettings = true;
