disabling UDP message transport when running a commandlet from the Editor unless explicitly allowed

diff --git a/Engine/Plugins/Messaging/UdpMessaging/Source/UdpMessaging/Private/UdpMessagingModule.cpp b/Engine/Plugins/Messaging/UdpMessaging/Source/UdpMessaging/Private/UdpMessagingModule.cpp
--- a/Engine/Plugins/Messaging/UdpMessaging/Source/UdpMessaging/Private/UdpMessagingModule.cpp
+++ b/Engine/Plugins/Messaging/UdpMessaging/Source/UdpMessaging/Private/UdpMessagingModule.cpp
@@ -22,19 +22,9 @@ public:
 
 	virtual void StartupModule( ) OVERRIDE
 	{
-		if (FApp::IsGame())
+		if (!SupportsNetworkedTransport())
 		{
-			// don't use UDP transports in shipping and testing configurations
-			if ((FApp::GetBuildConfiguration() == EBuildConfigurations::Shipping) || (FApp::GetBuildConfiguration() == EBuildConfigurations::Test))
-			{
-				return;
-			}
-
-			// don't use UDP transports in debug configurations unless supported and explicitly desired
-			if (!FParse::Param(FCommandLine::Get(), TEXT("Messaging")) || !FPlatformMisc::SupportsMessaging() || !FPlatformProcess::SupportsMultithreading())
-			{
-				return;
-			}
+			return;
 		}
 
 		// load dependencies
@@ -208,6 +198,42 @@ protected:
 		}
 	}
 
+	/**
+	 * Checks whether networked message transport is supported.
+	 *
+	 * @todo gmp: this should be moved into an Engine module, so it can be shared with other transports
+	 *
+	 * @return true if networked transport is supported, false otherwise.
+	 */
+	bool SupportsNetworkedTransport( ) const
+	{
+		if (FApp::IsGame())
+		{
+			// only allow UDP transport in Debug and Development configurations
+			if ((FApp::GetBuildConfiguration() == EBuildConfigurations::Shipping) || (FApp::GetBuildConfiguration() == EBuildConfigurations::Test))
+			{
+				return false;
+			}
+
+			// only allow UDP transport if the platform supports it
+			if (!FPlatformMisc::SupportsMessaging() || !FPlatformProcess::SupportsMultithreading())
+			{
+				return false;
+			}
+		}
+
+		if (FApp::IsGame() || IsRunningCommandlet())
+		{
+			// only allow UDP transport if explicitly desired
+			if (!FParse::Param(FCommandLine::Get(), TEXT("Messaging")))
+			{
+				return false;
+			}
+		}
+
+		return true;
+	}
+
 	/**
 	 * Shuts down the message bridge.
 	 */
