removed the remaining dependencies to the Networking module from the Messaging module;
cleaned up some fallout from incorrectly set up modules

diff --git a/Engine/Plugins/Messaging/UdpMessaging/Source/UdpMessaging/Classes/UdpMessagingSettings.h b/Engine/Plugins/Messaging/UdpMessaging/Source/UdpMessaging/Classes/UdpMessagingSettings.h
deleted file mode 100644
diff --git a/Engine/Plugins/Messaging/UdpMessaging/Source/UdpMessaging/Private/Shared/UdpMessageSegment.h b/Engine/Plugins/Messaging/UdpMessaging/Source/UdpMessaging/Private/Shared/UdpMessageSegment.h
new file mode 100644
--- /dev/null
+++ b/Engine/Plugins/Messaging/UdpMessaging/Source/UdpMessaging/Private/Shared/UdpMessageSegment.h
@@ -0,0 +1,285 @@
+// Copyright 1998-2014 Epic Games, Inc. All Rights Reserved.
+
+#pragma once
+
+
+namespace EUdpMessageSegments
+{
+	/**
+	 * Enumerates message segment types.
+	 */
+	enum Type
+	{
+		/**
+		 * None.
+		 */
+		None,
+
+		/**
+		 * Request to abort the sending of a message.
+		 */
+		Abort,
+
+		/**
+		 * Acknowledges that the message was received successfully.
+		 */
+		Acknowledge,
+
+		/**
+		 * Notifies the bus that an endpoint has left.
+		 */
+		Bye,
+
+		/**
+		 * A message data segment.
+		 */
+		Data,
+
+		/**
+		 * Notifies the bus that an endpoint has joined.
+		 */
+		Hello,
+
+		/**
+		 * Request to retransmit selected data segments.
+		 */
+		Retransmit,
+
+		/**
+		 * Notification that an inbound message timed out.
+		 */
+		Timeout
+	};
+}
+
+
+namespace FUdpMessageSegment
+{
+	/**
+	 * Structure for the header of all segments.
+	 */
+	struct FHeader
+	{
+		/**
+		 * Holds the protocol version.
+		 */
+		uint8 ProtocolVersion;
+
+		/**
+		 * Holds the recipient's node identifier (empty = multicast).
+		 */
+		FGuid RecipientNodeId;
+
+		/**
+		 * Holds the sender's node identifier.
+		 */
+		FGuid SenderNodeId;
+
+		/**
+		 * Holds the segment type.
+		 *
+		 * @see EUdpMessageSegments
+		 */
+		uint8 SegmentType;
+
+
+	public:
+
+		/**
+		 * Serializes the given header from or into the specified archive.
+		 *
+		 * @param Ar - The archive to serialize from or into.
+		 * @param DateTime - The header to serialize.
+		 *
+		 * @return The archive.
+		 */
+		friend FArchive& operator<< (FArchive& Ar, FHeader& Header)
+		{
+			return Ar << Header.ProtocolVersion << Header.RecipientNodeId << Header.SenderNodeId << Header.SegmentType;
+		}
+	};
+
+
+	/**
+	 * Structure for the sub-header of Abort segments.
+	 */
+	struct FAbortChunk
+	{
+		/** 
+		 * Holds the identifier of the message to abort.
+		 */
+		int32 MessageId;
+
+
+	public:
+
+		/**
+		 * Serializes the given header from or to the specified archive.
+		 *
+		 * @param Ar - The archive to serialize from or into.
+		 * @param DateTime - The header to serialize.
+		 *
+		 * @return The archive.
+		 */
+		friend FArchive& operator<< (FArchive& Ar, FAbortChunk& Header)
+		{
+			return Ar << Header.MessageId;
+		}
+	};
+
+
+	/**
+	 * Structure for the header of Acknowledge segments.
+	 */
+	struct FAcknowledgeChunk
+	{
+		/** 
+		 * Holds the identifier of the message that was received successfully.
+		 */
+		int32 MessageId;
+
+
+	public:
+
+		/**
+		 * Serializes the given header from or into the specified archive.
+		 *
+		 * @param Ar - The archive to serialize from or into.
+		 * @param DateTime - The header to serialize.
+		 *
+		 * @return The archive.
+		 */
+		friend FArchive& operator<< (FArchive& Ar, FAcknowledgeChunk& Header)
+		{
+			return Ar << Header.MessageId;
+		}
+	};
+
+
+	/**
+	 * Structure for the header of Data segments.
+	 */
+	struct FDataChunk
+	{
+		/**
+		 * Holds the identifier of the message that the data belongs to.
+		 */
+		int32 MessageId;
+
+		/**
+		 * Holds the total size of the message.
+		 */
+		int32 MessageSize;
+
+		/**
+		 * Holds the sequence number of this segment.
+		 */
+		uint16 SegmentNumber;
+
+		/**
+		 * Holds the segment's offset within the message.
+		 */
+		uint32 SegmentOffset;
+
+		/**
+		 * Holds the message sequence number (0 = not sequential).
+		 */
+		uint64 Sequence;
+
+		/**
+		 * Holds the total number of data segments being sent.
+		 */
+		uint16 TotalSegments;
+
+		/**
+		 * Holds the segment data.
+		 */
+		TArray<uint8> Data;
+
+
+	public:
+
+		/**
+		 * Serializes the given header from or into the specified archive.
+		 *
+		 * @param Ar - The archive to serialize from or into.
+		 * @param DateTime - The header to serialize.
+		 *
+		 * @return The archive.
+		 */
+		friend FArchive& operator<< (FArchive& Ar, FDataChunk& Chunk)
+		{
+			return Ar
+				<< Chunk.MessageId
+				<< Chunk.MessageSize
+				<< Chunk.SegmentNumber
+				<< Chunk.SegmentOffset
+				<< Chunk.Sequence
+				<< Chunk.TotalSegments
+				<< Chunk.Data;
+		}
+	};
+
+
+	/**
+	 * Structure for the sub-header of Retransmit segments.
+	 *
+	 * Retransmit segments are sent from a message recipient to a message sender in order
+	 * to indicate that selected message segments are to be retransmitted, i.e. if they
+	 * got lost on the network or if the recipient was unable to handle them previously.
+	 */
+	struct FRetransmitChunk
+	{
+		/**
+		 * Holds the identifier of the message for which data needs to be retransmitted.
+		 */
+		int32 MessageId;
+
+		/**
+		 * Holds the list of data segments that need to be retransmitted.
+		 */
+		TArray<uint16> Segments;
+
+	public:
+
+		/**
+		 * Serializes the given header from or into the specified archive.
+		 *
+		 * @param Ar - The archive to serialize from or into.
+		 * @param DateTime - The header to serialize.
+		 *
+		 * @return The archive.
+		 */
+		friend FArchive& operator<< (FArchive& Ar, FRetransmitChunk& Header)
+		{
+			return Ar << Header.MessageId << Header.Segments;
+		}
+	};
+
+
+	/**
+	 * Structure for the header of Timeout packets.
+	 */
+	struct FTimeoutChunk
+	{
+		/** 
+		 * Holds the identifier of the message that timed out.
+		 */
+		int32 MessageId;
+
+	public:
+
+		/**
+		 * Serializes the given header from or into the specified archive.
+		 *
+		 * @param Ar - The archive to serialize from or into.
+		 * @param DateTime - The header to serialize.
+		 *
+		 * @return The archive.
+		 */
+		friend FArchive& operator<< (FArchive& Ar, FTimeoutChunk& Header)
+		{
+			return Ar << Header.MessageId;
+		}
+	};
+};
\ No newline at end of file
diff --git a/Engine/Plugins/Messaging/UdpMessaging/Source/UdpMessaging/Private/Shared/UdpMessagingSettings.h b/Engine/Plugins/Messaging/UdpMessaging/Source/UdpMessaging/Private/Shared/UdpMessagingSettings.h
new file mode 100644
--- /dev/null
+++ b/Engine/Plugins/Messaging/UdpMessaging/Source/UdpMessaging/Private/Shared/UdpMessagingSettings.h
@@ -0,0 +1,79 @@
+// Copyright 1998-2014 Epic Games, Inc. All Rights Reserved.
+
+/*=============================================================================
+	UdpMessagingSettings.h: Declares the UUdpMessagingSettings class.
+=============================================================================*/
+
+#pragma once
+
+#include "UdpMessagingSettings.generated.h"
+
+
+UCLASS(config=Engine)
+class UUdpMessagingSettings
+	: public UObject
+{
+	GENERATED_UCLASS_BODY()
+
+public:
+
+	/**
+	 * Whether the UDP transport channel is enabled.
+	 */
+	UPROPERTY(config, EditAnywhere, Category=Transport)
+	bool EnableTransport;
+
+	/**
+	 * The IP endpoint to listen to and send packets from.
+	 */
+	UPROPERTY(config, EditAnywhere, Category=Transport)
+	FString UnicastEndpoint;
+
+	/**
+	 * The IP endpoint to send multicast packets to.
+	 */
+	UPROPERTY(config, EditAnywhere, Category=Transport)
+	FString MulticastEndpoint;
+
+	/**
+	 * The time-to-live (TTL) for sent multicast packets.
+	 */
+	UPROPERTY(config, EditAnywhere, Category=Transport)
+	uint8 MulticastTimeToLive;
+
+	/**
+	 * The IP endpoints of static devices.
+	 *
+	 * Use this setting to list devices on other subnets, such as mobile phones on a WiFi network.
+	 */
+	UPROPERTY(config, EditAnywhere, Category=Transport, AdvancedDisplay)
+	TArray<FString> StaticEndpoints;
+
+public:
+
+	/**
+	 * Whether the UDP tunnel is enabled.
+	 */
+	UPROPERTY(config, EditAnywhere, Category=Tunnel)
+	bool EnableTunnel;
+
+	/**
+	 * The IP endpoint to listen to and send packets from.
+	 */
+	UPROPERTY(config, EditAnywhere, Category=Tunnel)
+	FString TunnelUnicastEndpoint;
+
+	/**
+	 * The IP endpoint to send multicast packets to.
+	 */
+	UPROPERTY(config, EditAnywhere, Category=Tunnel)
+	FString TunnelMulticastEndpoint;
+
+	/**
+	 * The IP endpoints of remote tunnel nodes.
+	 *
+	 * Use this setting to connect to remote tunnel services.
+	 */
+	UPROPERTY(config, EditAnywhere, Category=Tunnel, AdvancedDisplay)
+	TArray<FString> RemoteTunnelEndpoints;
+};
diff --git a/Engine/Plugins/Messaging/UdpMessaging/Source/UdpMessaging/Private/Tunnel/UdpMessageTunnel.h b/Engine/Plugins/Messaging/UdpMessaging/Source/UdpMessaging/Private/Tunnel/UdpMessageTunnel.h
--- a/Engine/Plugins/Messaging/UdpMessaging/Source/UdpMessaging/Private/Tunnel/UdpMessageTunnel.h
+++ b/Engine/Plugins/Messaging/UdpMessaging/Source/UdpMessaging/Private/Tunnel/UdpMessageTunnel.h
@@ -12,7 +12,7 @@
  */
 class FUdpMessageTunnel
 	: FRunnable
-	, public IMessageTunnel
+	, public IUdpMessageTunnel
 {
 	// Structure for transport node information
 	struct FNodeInfo
@@ -128,7 +128,7 @@ public:
 
 public:
 
-	// Begin IMessageTunnel interface
+	// Begin IUdpMessageTunnel interface
 
 	virtual bool Connect( const FIPv4Endpoint& RemoteEndpoint ) OVERRIDE
 	{
@@ -151,7 +151,7 @@ public:
 		return false;
 	}
 
-	virtual int32 GetConnections( TArray<IMessageTunnelConnectionPtr>& OutConnections ) OVERRIDE
+	virtual int32 GetConnections( TArray<IUdpMessageTunnelConnectionPtr>& OutConnections ) OVERRIDE
 	{
 		FScopeLock Lock(&CriticalSection);
 
@@ -197,7 +197,7 @@ public:
 		Listener = NULL;
 	}
 
-	// End IMessageTunnel interface
+	// End IUdpMessageTunnel interface
 
 protected:
 
diff --git a/Engine/Plugins/Messaging/UdpMessaging/Source/UdpMessaging/Private/Tunnel/UdpMessageTunnelConnection.h b/Engine/Plugins/Messaging/UdpMessaging/Source/UdpMessaging/Private/Tunnel/UdpMessageTunnelConnection.h
--- a/Engine/Plugins/Messaging/UdpMessaging/Source/UdpMessaging/Private/Tunnel/UdpMessageTunnelConnection.h
+++ b/Engine/Plugins/Messaging/UdpMessaging/Source/UdpMessaging/Private/Tunnel/UdpMessageTunnelConnection.h
@@ -44,7 +44,7 @@ DECLARE_DELEGATE_TwoParams(FOnUdpMessageTunnelConnectionPacketReceived, FSaveByt
 class FUdpMessageTunnelConnection
 	: public FRunnable
 	, public TSharedFromThis<FUdpMessageTunnelConnection>
-	, public IMessageTunnelConnection
+	, public IUdpMessageTunnelConnection
 {
 public:
 
@@ -153,7 +153,7 @@ public:
 
 public:
 
-	// Begin IMessageTunnelConnection interface
+	// Begin IUdpMessageTunnelConnection interface
 
 	virtual void Close( ) OVERRIDE
 	{
@@ -190,7 +190,7 @@ public:
 		return (Socket->GetConnectionState() == SCS_Connected);
 	}
 
-	// End IMessageTunnelConnection interface
+	// End IUdpMessageTunnelConnection interface
 
 protected:
 
diff --git a/Engine/Plugins/Messaging/UdpMessaging/Source/UdpMessaging/Private/UdpMessageSegment.h b/Engine/Plugins/Messaging/UdpMessaging/Source/UdpMessaging/Private/UdpMessageSegment.h
deleted file mode 100644
diff --git a/Engine/Plugins/Messaging/UdpMessaging/Source/UdpMessaging/Private/UdpMessagingModule.cpp b/Engine/Plugins/Messaging/UdpMessaging/Source/UdpMessaging/Private/UdpMessagingModule.cpp
--- a/Engine/Plugins/Messaging/UdpMessaging/Source/UdpMessaging/Private/UdpMessagingModule.cpp
+++ b/Engine/Plugins/Messaging/UdpMessaging/Source/UdpMessaging/Private/UdpMessagingModule.cpp
@@ -299,7 +299,7 @@ private:
 	IMessageBridgePtr MessageBridge;
 
 	// Holds the message tunnel if present
-	IMessageTunnelPtr MessageTunnel;
+	IUdpMessageTunnelPtr MessageTunnel;
 };
 
 
diff --git a/Engine/Plugins/Messaging/UdpMessaging/Source/UdpMessaging/Private/UdpMessagingPrivatePCH.h b/Engine/Plugins/Messaging/UdpMessaging/Source/UdpMessaging/Private/UdpMessagingPrivatePCH.h
--- a/Engine/Plugins/Messaging/UdpMessaging/Source/UdpMessaging/Private/UdpMessagingPrivatePCH.h
+++ b/Engine/Plugins/Messaging/UdpMessaging/Source/UdpMessaging/Private/UdpMessagingPrivatePCH.h
@@ -6,6 +6,8 @@
 
 #pragma once
 
+#include "UdpMessaging.h"
+
 
 /* Dependencies
  *****************************************************************************/
diff --git a/Engine/Plugins/Messaging/UdpMessaging/Source/UdpMessaging/Public/Interfaces/IUdpMessageTunnel.h b/Engine/Plugins/Messaging/UdpMessaging/Source/UdpMessaging/Public/Interfaces/IUdpMessageTunnel.h
new file mode 100644
--- /dev/null
+++ b/Engine/Plugins/Messaging/UdpMessaging/Source/UdpMessaging/Public/Interfaces/IUdpMessageTunnel.h
@@ -0,0 +1,107 @@
+// Copyright 1998-2014 Epic Games, Inc. All Rights Reserved.
+
+/*=============================================================================
+	IUdpMessageTunnel.h: Declares the IUdpMessageTunnel interface.
+=============================================================================*/
+
+#pragma once
+
+
+/**
+ * Type definition for shared pointers to instances of IUdpMessageTunnel.
+ */
+typedef TSharedPtr<class IUdpMessageTunnel> IUdpMessageTunnelPtr;
+
+/**
+ * Type definition for shared references to instances of IUdpMessageTunnel.
+ */
+typedef TSharedRef<class IUdpMessageTunnel> IUdpMessageTunnelRef;
+
+
+/**
+ * Interface for message tunnels.
+ */
+class IUdpMessageTunnel
+{
+public:
+
+	/**
+	 * Establishes a tunnel with a remote server.
+	 *
+	 * @param RemoteEndpoint - The endpoint of the server to connect to.
+	 *
+	 * @return true if the connection has been established, false otherwise.
+	 *
+	 * @see GetConnections
+	 */
+	virtual bool Connect( const FIPv4Endpoint& RemoteEndpoint ) = 0;
+
+	/** 
+	 * Starts the tunnel server.
+	 *
+	 * @param LocalEndpoint - The IP endpoint to listen for incoming connections on.
+	 *
+	 * @see IsServerRunning
+	 * @see StopServer
+	 */
+	virtual void StartServer( const FIPv4Endpoint& LocalEndpoint ) = 0;
+
+	/**
+	 * Stops the tunnel server.
+	 *
+	 * @see IsServerRunning
+	 * @see StartServer
+	 */
+	virtual void StopServer( ) = 0;
+
+public:
+
+	/**
+	 * Gets the list of all open tunnel connections.
+	 *
+	 * @param OutConnections - Will hold the list of connections.
+	 *
+	 * @return The number of connections returned.
+	 *
+	 * @see Connect
+	 */
+	virtual int32 GetConnections( TArray<IUdpMessageTunnelConnectionPtr>& OutConnections ) = 0;
+
+	/**
+	 * Gets the total number of bytes that were received from tunnels.
+	 *
+	 * @return Number of bytes.
+	 */
+	virtual uint64 GetTotalInboundBytes( ) const = 0;
+
+	/**
+	 * Gets the total number of bytes that were sent out through tunnels.
+	 *
+	 * @return Number of bytes.
+	 */
+	virtual uint64 GetTotalOutboundBytes( ) const = 0;
+
+	/**
+	 * Checks whether the tunnel server is running.
+	 *
+	 * @return true if the tunnel server is running, false otherwise.
+	 *
+	 * @see StartServer
+	 * @see StopServer
+	 */
+	virtual bool IsServerRunning( ) const = 0;
+
+public:
+
+	/**
+	 * Gets a delegate that is executed when the list of incoming connections changed.
+	 */
+	virtual FSimpleDelegate& OnConnectionsChanged( ) = 0;
+
+public:
+
+	/**
+	 * Virtual destructor.
+	 */
+	virtual ~IUdpMessageTunnel( ) { }
+};
diff --git a/Engine/Plugins/Messaging/UdpMessaging/Source/UdpMessaging/Public/Interfaces/IUdpMessageTunnelConnection.h b/Engine/Plugins/Messaging/UdpMessaging/Source/UdpMessaging/Public/Interfaces/IUdpMessageTunnelConnection.h
new file mode 100644
--- /dev/null
+++ b/Engine/Plugins/Messaging/UdpMessaging/Source/UdpMessaging/Public/Interfaces/IUdpMessageTunnelConnection.h
@@ -0,0 +1,74 @@
+// Copyright 1998-2014 Epic Games, Inc. All Rights Reserved.
+
+/*=============================================================================
+	IUdpMessageTunnelConnection.h: Declares the IUdpMessageTunnelConnection interface.
+=============================================================================*/
+
+#pragma once
+
+
+/**
+ * Type definition for shared pointers to instances of IUdpMessageTunnelConnection.
+ */
+typedef TSharedPtr<class IUdpMessageTunnelConnection> IUdpMessageTunnelConnectionPtr;
+
+/**
+ * Type definition for shared references to instances of IUdpMessageTunnelConnection.
+ */
+typedef TSharedRef<class IUdpMessageTunnelConnection> IUdpMessageTunnelConnectionRef;
+
+
+/**
+ * Interface for message tunnel connections.
+ */
+class IUdpMessageTunnelConnection
+{
+public:
+
+	/**
+	 * Closes this connection.
+	 */
+	virtual void Close( ) = 0;
+
+	/**
+	 * Gets the total number of bytes received from this connection.
+	 *
+	 * @return Number of bytes.
+	 */
+	virtual uint64 GetTotalBytesReceived( ) const = 0;
+
+	/**
+	 * Gets the total number of bytes received from this connection.
+	 *
+	 * @return Number of bytes.
+	 */
+	virtual uint64 GetTotalBytesSent( ) const = 0;
+
+	/**
+	 * Gets the human readable name of the connection.
+	 *
+	 * @return Connection name.
+	 */
+	virtual FText GetName( ) const = 0;
+
+	/**
+	 * Gets the amount of time that the connection has been established.
+	 *
+	 * @return Time span.
+	 */
+	virtual FTimespan GetUptime( ) const = 0;
+
+	/**
+	 * Checks whether this connection is open.
+	 *
+	 * @return true if the connection is open, false otherwise.
+	 */
+	virtual bool IsOpen( ) const = 0;
+
+public:
+
+	/**
+	 * Virtual destructor.
+	 */
+	virtual ~IUdpMessageTunnelConnection( ) { }
+};
diff --git a/Engine/Plugins/Messaging/UdpMessaging/Source/UdpMessaging/UdpMessaging.Build.cs b/Engine/Plugins/Messaging/UdpMessaging/Source/UdpMessaging/UdpMessaging.Build.cs
--- a/Engine/Plugins/Messaging/UdpMessaging/Source/UdpMessaging/UdpMessaging.Build.cs
+++ b/Engine/Plugins/Messaging/UdpMessaging/Source/UdpMessaging/UdpMessaging.Build.cs
@@ -21,6 +21,7 @@ namespace UnrealBuildTool.Rules
 			PrivateIncludePaths.AddRange(
 				new string[] {
 					"UdpMessaging/Private",
+					"UdpMessaging/Private/Shared",
 					"UdpMessaging/Private/Transport",
 					"UdpMessaging/Private/Transport/Tests",
 					"UdpMessaging/Private/Tunnel",
diff --git a/Engine/Plugins/Runtime/LiveEditorListenServer/Source/LiveEditorListenServer/LiveEditorListenServer.Build.cs b/Engine/Plugins/Runtime/LiveEditorListenServer/Source/LiveEditorListenServer/LiveEditorListenServer.Build.cs
--- a/Engine/Plugins/Runtime/LiveEditorListenServer/Source/LiveEditorListenServer/LiveEditorListenServer.Build.cs
+++ b/Engine/Plugins/Runtime/LiveEditorListenServer/Source/LiveEditorListenServer/LiveEditorListenServer.Build.cs
@@ -25,15 +25,14 @@ namespace UnrealBuildTool.Rules
                     "Core",
                     "CoreUObject",
 					"Engine",
-                    "Sockets",
-                    "Networking"
 				}
 				);
 
 			PrivateDependencyModuleNames.AddRange(
 				new string[]
 				{
-					// ... add private dependencies that you statically link with here ...
+					"Networking",
+                    "Sockets",
 				}
 				);
 
diff --git a/Engine/Plugins/Runtime/LiveEditorListenServer/Source/LiveEditorListenServer/Private/LiveEditorListenServerPrivatePCH.h b/Engine/Plugins/Runtime/LiveEditorListenServer/Source/LiveEditorListenServer/Private/LiveEditorListenServerPrivatePCH.h
--- a/Engine/Plugins/Runtime/LiveEditorListenServer/Source/LiveEditorListenServer/Private/LiveEditorListenServerPrivatePCH.h
+++ b/Engine/Plugins/Runtime/LiveEditorListenServer/Source/LiveEditorListenServer/Private/LiveEditorListenServerPrivatePCH.h
@@ -1,6 +1,7 @@
 // Copyright 1998-2014 Epic Games, Inc. All Rights Reserved.
 
 #include "Engine.h"
+#include "Networking.h"
 #include "LiveEditorListenServerPublic.h"
 #include "ILiveEditorListenServer.h"
 #include "LiveEditorListenServerEngineHooks.h"
diff --git a/Engine/Source/Developer/AutomationController/AutomationController.Build.cs b/Engine/Source/Developer/AutomationController/AutomationController.Build.cs
--- a/Engine/Source/Developer/AutomationController/AutomationController.Build.cs
+++ b/Engine/Source/Developer/AutomationController/AutomationController.Build.cs
@@ -18,7 +18,6 @@ namespace UnrealBuildTool.Rules
 				{
 					"AutomationMessages",
 					"CoreUObject",
-					"Networking",
 					"UnrealEdMessages",
 				}
 			);
diff --git a/Engine/Source/Developer/AutomationController/Private/AutomationControllerPrivatePCH.h b/Engine/Source/Developer/AutomationController/Private/AutomationControllerPrivatePCH.h
--- a/Engine/Source/Developer/AutomationController/Private/AutomationControllerPrivatePCH.h
+++ b/Engine/Source/Developer/AutomationController/Private/AutomationControllerPrivatePCH.h
@@ -14,6 +14,7 @@
  *****************************************************************************/
 
 #include "AutomationMessages.h"
+#include "Messaging.h"
 #include "UnrealEdMessages.h"
 
 
diff --git a/Engine/Source/Developer/AutomationController/Public/AutomationController.h b/Engine/Source/Developer/AutomationController/Public/AutomationController.h
--- a/Engine/Source/Developer/AutomationController/Public/AutomationController.h
+++ b/Engine/Source/Developer/AutomationController/Public/AutomationController.h
@@ -12,7 +12,6 @@
 
 #include "Core.h"
 #include "ModuleManager.h"
-#include "Messaging.h"
 #include "IFilter.h"
 #include "TextFilter.h"
 #include "FilterCollection.h"
diff --git a/Engine/Source/Developer/AutomationWindow/AutomationWindow.Build.cs b/Engine/Source/Developer/AutomationWindow/AutomationWindow.Build.cs
--- a/Engine/Source/Developer/AutomationWindow/AutomationWindow.Build.cs
+++ b/Engine/Source/Developer/AutomationWindow/AutomationWindow.Build.cs
@@ -21,7 +21,6 @@ namespace UnrealBuildTool.Rules
 					"DesktopPlatform",
                     "InputCore",
 					"Messaging",
-					"Networking",
 					"Slate",
                     "EditorStyle",
 				}
diff --git a/Engine/Source/Developer/ProfilerClient/ProfilerClient.Build.cs b/Engine/Source/Developer/ProfilerClient/ProfilerClient.Build.cs
--- a/Engine/Source/Developer/ProfilerClient/ProfilerClient.Build.cs
+++ b/Engine/Source/Developer/ProfilerClient/ProfilerClient.Build.cs
@@ -21,7 +21,6 @@ namespace UnrealBuildTool.Rules
 				new string[]
 				{
 					"CoreUObject",
-					"Networking",
 					"ProfilerMessages",
 				}
 			);
diff --git a/Engine/Source/Developer/ScreenShotComparisonTools/ScreenShotComparisonTools.Build.cs b/Engine/Source/Developer/ScreenShotComparisonTools/ScreenShotComparisonTools.Build.cs
--- a/Engine/Source/Developer/ScreenShotComparisonTools/ScreenShotComparisonTools.Build.cs
+++ b/Engine/Source/Developer/ScreenShotComparisonTools/ScreenShotComparisonTools.Build.cs
@@ -19,7 +19,6 @@ public class ScreenShotComparisonTools : ModuleRules
 			{
 				"AutomationMessages",
 				"CoreUObject",
-				"Networking",
 				"UnrealEdMessages",
 				"Slate",
                 "EditorStyle",
diff --git a/Engine/Source/Developer/SessionServices/Private/SessionServicesPrivatePCH.h b/Engine/Source/Developer/SessionServices/Private/SessionServicesPrivatePCH.h
--- a/Engine/Source/Developer/SessionServices/Private/SessionServicesPrivatePCH.h
+++ b/Engine/Source/Developer/SessionServices/Private/SessionServicesPrivatePCH.h
@@ -6,13 +6,14 @@
 
 #pragma once
 
+#include "SessionServices.h"
+
 
 /* Dependencies
  *****************************************************************************/
 
-#include "../Public/SessionServices.h"
-
 #include "EngineMessages.h"
+#include "Messaging.h"
 #include "SessionMessages.h"
 #include "Ticker.h"
 
diff --git a/Engine/Source/Developer/SessionServices/Public/SessionServices.h b/Engine/Source/Developer/SessionServices/Public/SessionServices.h
--- a/Engine/Source/Developer/SessionServices/Public/SessionServices.h
+++ b/Engine/Source/Developer/SessionServices/Public/SessionServices.h
@@ -12,7 +12,6 @@
 
 #include "Core.h"
 #include "ModuleInterface.h"
-#include "Messaging.h"
 
 
 /* Interfaces
diff --git a/Engine/Source/Developer/SessionServices/SessionServices.Build.cs b/Engine/Source/Developer/SessionServices/SessionServices.Build.cs
--- a/Engine/Source/Developer/SessionServices/SessionServices.Build.cs
+++ b/Engine/Source/Developer/SessionServices/SessionServices.Build.cs
@@ -10,7 +10,6 @@ public class SessionServices : ModuleRules
 			new string[]
 			{
 				"Core",
-				"Networking",
 			}
 		);
 
diff --git a/Engine/Source/Developer/iOS/IOSTargetPlatform/IOSTargetPlatform.Build.cs b/Engine/Source/Developer/iOS/IOSTargetPlatform/IOSTargetPlatform.Build.cs
--- a/Engine/Source/Developer/iOS/IOSTargetPlatform/IOSTargetPlatform.Build.cs
+++ b/Engine/Source/Developer/iOS/IOSTargetPlatform/IOSTargetPlatform.Build.cs
@@ -17,7 +17,6 @@ public class IOSTargetPlatform : ModuleRules
 				"Messaging",
 				"TargetDeviceServices",
 				"LaunchDaemonMessages",
-				"Networking"
 			}
 		);
         PlatformSpecificDynamicallyLoadedModuleNames.Add("LaunchDaemonMessages");
diff --git a/Engine/Source/Editor/OnlineBlueprintSupport/OnlineBlueprintSupport.Build.cs b/Engine/Source/Editor/OnlineBlueprintSupport/OnlineBlueprintSupport.Build.cs
--- a/Engine/Source/Editor/OnlineBlueprintSupport/OnlineBlueprintSupport.Build.cs
+++ b/Engine/Source/Editor/OnlineBlueprintSupport/OnlineBlueprintSupport.Build.cs
@@ -13,7 +13,8 @@ public class OnlineBlueprintSupport : ModuleRules
 				"Engine",
 				"BlueprintGraph",
 				"OnlineSubsystem",
-				"OnlineSubsystemUtils"
+				"OnlineSubsystemUtils",
+				"Sockets",
 			}
 		);
 	}
diff --git a/Engine/Source/Runtime/AutomationWorker/AutomationWorker.Build.cs b/Engine/Source/Runtime/AutomationWorker/AutomationWorker.Build.cs
--- a/Engine/Source/Runtime/AutomationWorker/AutomationWorker.Build.cs
+++ b/Engine/Source/Runtime/AutomationWorker/AutomationWorker.Build.cs
@@ -18,7 +18,6 @@ namespace UnrealBuildTool.Rules
 				{
 					"AutomationMessages",
 					"CoreUObject",
-					"Networking",
 				}
 			);
 
diff --git a/Engine/Source/Runtime/Messaging/Messaging.Build.cs b/Engine/Source/Runtime/Messaging/Messaging.Build.cs
--- a/Engine/Source/Runtime/Messaging/Messaging.Build.cs
+++ b/Engine/Source/Runtime/Messaging/Messaging.Build.cs
@@ -10,7 +10,6 @@ namespace UnrealBuildTool.Rules
 				new string[]
 				{
 					"CoreUObject",
-					"Networking",
 				}
 			); 
 			
diff --git a/Engine/Source/Runtime/Messaging/Public/Interfaces/IMessageTunnel.h b/Engine/Source/Runtime/Messaging/Public/Interfaces/IMessageTunnel.h
deleted file mode 100644
diff --git a/Engine/Source/Runtime/Messaging/Public/Interfaces/IMessageTunnelConnection.h b/Engine/Source/Runtime/Messaging/Public/Interfaces/IMessageTunnelConnection.h
deleted file mode 100644
diff --git a/Engine/Source/Runtime/Messaging/Public/Messaging.h b/Engine/Source/Runtime/Messaging/Public/Messaging.h
--- a/Engine/Source/Runtime/Messaging/Public/Messaging.h
+++ b/Engine/Source/Runtime/Messaging/Public/Messaging.h
@@ -12,7 +12,6 @@
 
 #include "CoreUObject.h"
 #include "ModuleManager.h"
-#include "Networking.h"
 #include "TaskGraphInterfaces.h"
 
 
@@ -35,8 +34,6 @@
 #include "ISerializeMessages.h"
 #include "ITransportMessages.h"
 #include "IMessageBridge.h"
-#include "IMessageTunnelConnection.h"
-#include "IMessageTunnel.h"
 #include "IMessagingModule.h"
 
 
diff --git a/Engine/Source/Runtime/Online/OnlineSubsystemUtils/Private/OnlineSubsystemUtilsPrivatePCH.h b/Engine/Source/Runtime/Online/OnlineSubsystemUtils/Private/OnlineSubsystemUtilsPrivatePCH.h
--- a/Engine/Source/Runtime/Online/OnlineSubsystemUtils/Private/OnlineSubsystemUtilsPrivatePCH.h
+++ b/Engine/Source/Runtime/Online/OnlineSubsystemUtils/Private/OnlineSubsystemUtilsPrivatePCH.h
@@ -4,6 +4,7 @@
 
 #include "Core.h"
 #include "Engine.h"
+#include "Sockets.h"
 #include "OnlineSubsystem.h"
 #include "OnlineSubsystemUtils.h"
 #include "OnlineSubsystemUtilsModule.h"
