forgot to check in

diff --git a/Engine/Plugins/Messaging/UdpMessaging/Source/UdpMessaging/Public/UdpMessaging.h b/Engine/Plugins/Messaging/UdpMessaging/Source/UdpMessaging/Public/UdpMessaging.h
new file mode 100644
--- /dev/null
+++ b/Engine/Plugins/Messaging/UdpMessaging/Source/UdpMessaging/Public/UdpMessaging.h
@@ -0,0 +1,20 @@
+// Copyright 1998-2014 Epic Games, Inc. All Rights Reserved.
+
+/*=============================================================================
+	UdpMessaging.h: UdpMessaging module public header file.
+=============================================================================*/
+
+#pragma once
+
+
+/* Dependencies
+ *****************************************************************************/
+
+#include "Networking.h"
+
+
+/* Interfaces
+ *****************************************************************************/
+
+#include "IUdpMessageTunnelConnection.h"
+#include "IUdpMessageTunnel.h"
