fixed up obsolete module dependencies in Messaging plugins

diff --git a/Engine/Plugins/Messaging/MessagingDebugger/Source/MessagingDebugger/MessagingDebugger.Build.cs b/Engine/Plugins/Messaging/MessagingDebugger/Source/MessagingDebugger/MessagingDebugger.Build.cs
--- a/Engine/Plugins/Messaging/MessagingDebugger/Source/MessagingDebugger/MessagingDebugger.Build.cs
+++ b/Engine/Plugins/Messaging/MessagingDebugger/Source/MessagingDebugger/MessagingDebugger.Build.cs
@@ -12,13 +12,18 @@ public class MessagingDebugger : ModuleRules
 				"Core",
 				"CoreUObject",
                 "InputCore",
-				"Messaging",
                 "PropertyEditor",
 				"Slate",
                 "EditorStyle"
 			}
 		);
 
+		PrivateIncludePathModuleNames.AddRange(
+			new string[] {
+				"Messaging",
+			}
+		);
+
 		PrivateIncludePaths.AddRange(
 			new string[] {
 				"MessagingDebugger/Private",
diff --git a/Engine/Plugins/Messaging/UdpMessaging/Source/UdpMessaging/UdpMessaging.Build.cs b/Engine/Plugins/Messaging/UdpMessaging/Source/UdpMessaging/UdpMessaging.Build.cs
--- a/Engine/Plugins/Messaging/UdpMessaging/Source/UdpMessaging/UdpMessaging.Build.cs
+++ b/Engine/Plugins/Messaging/UdpMessaging/Source/UdpMessaging/UdpMessaging.Build.cs
@@ -6,15 +6,25 @@ namespace UnrealBuildTool.Rules
 	{
 		public UdpMessaging(TargetInfo Target)
 		{
-			PrivateDependencyModuleNames.AddRange(
+			PublicDependencyModuleNames.AddRange(
 				new string[]
 				{
+					"Networking",
+				}
+			); 
+
+			PrivateDependencyModuleNames.AddRange(
+				new string[] {
 					"Core",
 					"CoreUObject",
+					"Sockets",
+				}
+			);
+
+			PrivateIncludePathModuleNames.AddRange(
+				new string[] {
 					"Messaging",
-					"Networking",
 					"Settings",
-					"Sockets",
 				}
 			);
 
