made UdpMessaging plug-in application life cycle aware; enabled hot reloading

diff --git a/Engine/Plugins/Messaging/UdpMessaging/Source/UdpMessaging/Private/UdpMessagingModule.cpp b/Engine/Plugins/Messaging/UdpMessaging/Source/UdpMessaging/Private/UdpMessagingModule.cpp
--- a/Engine/Plugins/Messaging/UdpMessaging/Source/UdpMessaging/Private/UdpMessagingModule.cpp
+++ b/Engine/Plugins/Messaging/UdpMessaging/Source/UdpMessaging/Private/UdpMessagingModule.cpp
@@ -51,22 +51,19 @@ public:
 			);
 		}
 
-		// start up services if desired
-		const UUdpMessagingSettings& Settings = *GetDefault<UUdpMessagingSettings>();
-
-		if (Settings.EnableTransport)
-		{
-			InitializeBridge();
-		}
+		// register application events
+		FCoreDelegates::ApplicationHasReactivatedDelegate.AddRaw(this, &FUdpMessagingModule::HandleApplicationHasReactivated);
+		FCoreDelegates::ApplicationWillDeactivateDelegate.AddRaw(this, &FUdpMessagingModule::HandleApplicationWillDeactivate);
 
-		if (Settings.EnableTunnel)
-		{
-			InitializeTunnel();
-		}
+		RestartServices();
 	}
 
 	virtual void ShutdownModule( ) OVERRIDE
 	{
+		// unregister application events
+		FCoreDelegates::ApplicationHasReactivatedDelegate.RemoveAll(this);
+		FCoreDelegates::ApplicationWillDeactivateDelegate.RemoveAll(this);
+
 		// unregister settings
 		ISettingsModule* SettingsModule = ISettingsModule::Get();
 
@@ -82,7 +79,7 @@ public:
 
 	virtual bool SupportsDynamicReloading( ) OVERRIDE
 	{
-		return false;
+		return true;
 	}
 
 	// End IModuleInterface interface
@@ -198,6 +195,38 @@ protected:
 		}
 	}
 
+	/**
+	 * Restarts the bridge and tunnel services.
+	 */
+	void RestartServices( )
+	{
+		const UUdpMessagingSettings& Settings = *GetDefault<UUdpMessagingSettings>();
+
+		if (Settings.EnableTransport)
+		{
+			if (!MessageBridge.IsValid())
+			{
+				InitializeBridge();
+			}
+		}
+		else
+		{
+			ShutdownBridge();
+		}
+
+		if (Settings.EnableTunnel)
+		{
+			if (!MessageTunnel.IsValid())
+			{
+				InitializeTunnel();
+			}
+		}
+		else
+		{
+			ShutdownTunnel();
+		}
+	}
+
 	/**
 	 * Checks whether networked message transport is supported.
 	 *
@@ -261,34 +290,23 @@ protected:
 
 private:
 
-	// Handles saved settings.
-	bool HandleSettingsSaved( )
+	// Callback for when an has been reactivated (i.e. return from sleep on iOS).
+	void HandleApplicationHasReactivated( )
 	{
-		const UUdpMessagingSettings& Settings = *GetDefault<UUdpMessagingSettings>();
+		RestartServices();
+	}
 
-		if (Settings.EnableTransport)
-		{
-			if (!MessageBridge.IsValid())
-			{
-				InitializeBridge();
-			}
-		}
-		else
-		{
-			ShutdownBridge();
-		}
+	// Callback for when the application will be deactivated (i.e. sleep on iOS).
+	void HandleApplicationWillDeactivate( )
+	{
+		ShutdownBridge();
+		ShutdownTunnel();
+	}
 
-		if (Settings.EnableTunnel)
-		{
-			if (!MessageTunnel.IsValid())
-			{
-				InitializeTunnel();
-			}
-		}
-		else
-		{
-			ShutdownTunnel();
-		}
+	// Callback for when the settings were saved.
+	bool HandleSettingsSaved( )
+	{
+		RestartServices();
 
 		return true;
 	}
