removed pending data size checks (no longer needed, because HasPendingData() now correctly returns false if there is no pending data)

diff --git a/Engine/Plugins/Messaging/UdpMessaging/Source/UdpMessaging/Private/Tunnel/UdpMessageTunnel.h b/Engine/Plugins/Messaging/UdpMessaging/Source/UdpMessaging/Private/Tunnel/UdpMessageTunnel.h
--- a/Engine/Plugins/Messaging/UdpMessaging/Source/UdpMessaging/Private/Tunnel/UdpMessageTunnel.h
+++ b/Engine/Plugins/Messaging/UdpMessaging/Source/UdpMessaging/Private/Tunnel/UdpMessageTunnel.h
@@ -289,57 +289,54 @@ protected:
 
 		while (Socket->HasPendingData(DatagramSize))
 		{
-			if (DatagramSize > 0)
-			{
-				FArrayReaderPtr Datagram = MakeShareable(new FArrayReader(true));
-				Datagram->Init(FMath::Min(DatagramSize, 65507u));
+			FArrayReaderPtr Datagram = MakeShareable(new FArrayReader(true));
+			Datagram->Init(FMath::Min(DatagramSize, 65507u));
+
+			int32 BytesRead = 0;
 
-				int32 BytesRead = 0;
+			if (Socket->RecvFrom(Datagram->GetData(), Datagram->Num(), BytesRead, *Sender))
+			{
+				FUdpMessageSegment::FHeader Header;
+				*Datagram << Header;
 
-				if (Socket->RecvFrom(Datagram->GetData(), Datagram->Num(), BytesRead, *Sender))
+				// check protocol version
+				if (Header.ProtocolVersion != UDP_MESSAGING_TRANSPORT_PROTOCOL_VERSION)
 				{
-					FUdpMessageSegment::FHeader Header;
-					*Datagram << Header;
+					return;
+				}
 
-					// check protocol version
-					if (Header.ProtocolVersion != UDP_MESSAGING_TRANSPORT_PROTOCOL_VERSION)
-					{
-						return;
-					}
+				// ignore loopback datagrams
+				if (RemoteNodes.Contains(Header.SenderNodeId))
+				{
+					return;
+				}
 
-					// ignore loopback datagrams
-					if (RemoteNodes.Contains(Header.SenderNodeId))
-					{
-						return;
-					}
+				Datagram->Seek(0);
 
-					Datagram->Seek(0);
+				// forward datagram to remote nodes
+				if (Header.RecipientNodeId.IsValid())
+				{
+					FNodeInfo* RemoteNode = RemoteNodes.Find(Header.RecipientNodeId);
 
-					// forward datagram to remote nodes
-					if (Header.RecipientNodeId.IsValid())
+					if ((RemoteNode != NULL) && (RemoteNode->Connection.IsValid()))
 					{
-						FNodeInfo* RemoteNode = RemoteNodes.Find(Header.RecipientNodeId);
-
-						if ((RemoteNode != NULL) && (RemoteNode->Connection.IsValid()))
-						{
-							RemoteNode->Connection->Send(Datagram);
-						}
+						RemoteNode->Connection->Send(Datagram);
 					}
-					else
+				}
+				else
+				{
+					for (int32 ConnectionIndex = 0; ConnectionIndex < Connections.Num(); ++ConnectionIndex)
 					{
-						for (int32 ConnectionIndex = 0; ConnectionIndex < Connections.Num(); ++ConnectionIndex)
-						{
-							Connections[ConnectionIndex]->Send(Datagram);
-						}
+						Connections[ConnectionIndex]->Send(Datagram);
 					}
+				}
 
-					// update local node & statistics
-					FNodeInfo& LocalNode = LocalNodes.FindOrAdd(Header.SenderNodeId);
-					LocalNode.Endpoint = FIPv4Endpoint(Sender);
-					LocalNode.LastDatagramReceivedTime = CurrentTime;
+				// update local node & statistics
+				FNodeInfo& LocalNode = LocalNodes.FindOrAdd(Header.SenderNodeId);
+				LocalNode.Endpoint = FIPv4Endpoint(Sender);
+				LocalNode.LastDatagramReceivedTime = CurrentTime;
 
-					TotalOutboundBytes += Datagram->Num();
-				}
+				TotalOutboundBytes += Datagram->Num();
 			}
 		}
 	}
diff --git a/Engine/Plugins/Runtime/LiveEditorListenServer/Source/LiveEditorListenServer/Private/LiveEditorListenServer.cpp b/Engine/Plugins/Runtime/LiveEditorListenServer/Source/LiveEditorListenServer/Private/LiveEditorListenServer.cpp
--- a/Engine/Plugins/Runtime/LiveEditorListenServer/Source/LiveEditorListenServer/Private/LiveEditorListenServer.cpp
+++ b/Engine/Plugins/Runtime/LiveEditorListenServer/Source/LiveEditorListenServer/Private/LiveEditorListenServer.cpp
@@ -187,7 +187,7 @@ void FLiveEditorListenServer::Tick( float DeltaTime )
 	{
 		FSocket *Client = *ClientIt;
 		uint32 DataSize = 0;
-		while ( Client->HasPendingData(DataSize) && DataSize > 0 )
+		while ( Client->HasPendingData(DataSize) )
 		{
 			FArrayReaderPtr Datagram = MakeShareable(new FArrayReader(true));
 			Datagram->Init(FMath::Min(DataSize, 65507u));
diff --git a/Engine/Source/Runtime/Networking/Public/Common/UdpSocketReceiver.h b/Engine/Source/Runtime/Networking/Public/Common/UdpSocketReceiver.h
--- a/Engine/Source/Runtime/Networking/Public/Common/UdpSocketReceiver.h
+++ b/Engine/Source/Runtime/Networking/Public/Common/UdpSocketReceiver.h
@@ -96,17 +96,13 @@ public:
 
 			while (Socket->HasPendingData(Size))
 			{
-				if (Size > 0)
-				{
-					FArrayReaderPtr Reader = MakeShareable(new FArrayReader(true));
-					Reader->Init(FMath::Min(Size, 65507u));
+				FArrayReaderPtr Reader = MakeShareable(new FArrayReader(true));
+				Reader->Init(FMath::Min(Size, 65507u));
 
-					int32 Read = 0;
+				int32 Read = 0;
+				Socket->RecvFrom(Reader->GetData(), Reader->Num(), Read, *Sender);
 
-					Socket->RecvFrom(Reader->GetData(), Reader->Num(), Read, *Sender);
-
-					DataReceivedDelegate.ExecuteIfBound(Reader, FIPv4Endpoint(Sender));
-				}
+				DataReceivedDelegate.ExecuteIfBound(Reader, FIPv4Endpoint(Sender));
 			}
 		}
 
