UdpMessaging: documentation & nullptr cleanup

diff --git a/Engine/Plugins/Messaging/UdpMessaging/Source/UdpMessaging/Private/Shared/UdpMessageSegment.h b/Engine/Plugins/Messaging/UdpMessaging/Source/UdpMessaging/Private/Shared/UdpMessageSegment.h
--- a/Engine/Plugins/Messaging/UdpMessaging/Source/UdpMessaging/Private/Shared/UdpMessageSegment.h
+++ b/Engine/Plugins/Messaging/UdpMessaging/Source/UdpMessaging/Private/Shared/UdpMessageSegment.h
@@ -82,14 +82,13 @@ namespace FUdpMessageSegment
 		 */
 		uint8 SegmentType;
 
-
 	public:
 
 		/**
 		 * Serializes the given header from or into the specified archive.
 		 *
-		 * @param Ar - The archive to serialize from or into.
-		 * @param DateTime - The header to serialize.
+		 * @param Ar The archive to serialize from or into.
+		 * @param DateTime The header to serialize.
 		 *
 		 * @return The archive.
 		 */
@@ -110,14 +109,13 @@ namespace FUdpMessageSegment
 		 */
 		int32 MessageId;
 
-
 	public:
 
 		/**
 		 * Serializes the given header from or to the specified archive.
 		 *
-		 * @param Ar - The archive to serialize from or into.
-		 * @param DateTime - The header to serialize.
+		 * @param Ar The archive to serialize from or into.
+		 * @param DateTime The header to serialize.
 		 *
 		 * @return The archive.
 		 */
@@ -138,14 +136,13 @@ namespace FUdpMessageSegment
 		 */
 		int32 MessageId;
 
-
 	public:
 
 		/**
 		 * Serializes the given header from or into the specified archive.
 		 *
-		 * @param Ar - The archive to serialize from or into.
-		 * @param DateTime - The header to serialize.
+		 * @param Ar The archive to serialize from or into.
+		 * @param DateTime The header to serialize.
 		 *
 		 * @return The archive.
 		 */
@@ -196,14 +193,13 @@ namespace FUdpMessageSegment
 		 */
 		TArray<uint8> Data;
 
-
 	public:
 
 		/**
 		 * Serializes the given header from or into the specified archive.
 		 *
-		 * @param Ar - The archive to serialize from or into.
-		 * @param DateTime - The header to serialize.
+		 * @param Ar The archive to serialize from or into.
+		 * @param DateTime The header to serialize.
 		 *
 		 * @return The archive.
 		 */
@@ -245,8 +241,8 @@ namespace FUdpMessageSegment
 		/**
 		 * Serializes the given header from or into the specified archive.
 		 *
-		 * @param Ar - The archive to serialize from or into.
-		 * @param DateTime - The header to serialize.
+		 * @param Ar The archive to serialize from or into.
+		 * @param DateTime The header to serialize.
 		 *
 		 * @return The archive.
 		 */
@@ -272,8 +268,8 @@ namespace FUdpMessageSegment
 		/**
 		 * Serializes the given header from or into the specified archive.
 		 *
-		 * @param Ar - The archive to serialize from or into.
-		 * @param DateTime - The header to serialize.
+		 * @param Ar The archive to serialize from or into.
+		 * @param DateTime The header to serialize.
 		 *
 		 * @return The archive.
 		 */
diff --git a/Engine/Plugins/Messaging/UdpMessaging/Source/UdpMessaging/Private/Transport/ReassembledUdpMessage.h b/Engine/Plugins/Messaging/UdpMessaging/Source/UdpMessaging/Private/Transport/ReassembledUdpMessage.h
--- a/Engine/Plugins/Messaging/UdpMessaging/Source/UdpMessaging/Private/Transport/ReassembledUdpMessage.h
+++ b/Engine/Plugins/Messaging/UdpMessaging/Source/UdpMessaging/Private/Transport/ReassembledUdpMessage.h
@@ -22,10 +22,10 @@ public:
 	/**
 	 * Creates and initializes a new inbound message info.
 	 *
-	 * @param MessageSize - The total size of the message in bytes.
-	 * @param SegmentCount - The total number of segments that need to be received for this message.
-	 * @param InSequence - The message sequence number.
-	 * @param InSender - The IPv4 endpoint of the sender.
+	 * @param MessageSize The total size of the message in bytes.
+	 * @param SegmentCount The total number of segments that need to be received for this message.
+	 * @param InSequence The message sequence number.
+	 * @param InSender The IPv4 endpoint of the sender.
 	 */
 	FReassembledUdpMessage(int32 MessageSize, int32 SegmentCount, uint64 InSequence, const FIPv4Endpoint& InSender)
 		: PendingSegments(true, SegmentCount)
@@ -40,8 +40,7 @@ public:
 	/**
 	 * Virtual destructor.
 	 */
-	virtual ~FReassembledUdpMessage( )
-	{ }
+	virtual ~FReassembledUdpMessage( ) { }
 
 public:
 
@@ -118,10 +117,10 @@ public:
 	/**
 	 * Reassembles a segment into the specified message.
 	 *
-	 * @param SegmentNumber - The number of the message segment.
-	 * @param SegmentOffset - The segment's offset within the message.
-	 * @param SegmentData - The segment data.
-	 * @param CurrentTime - The current time.
+	 * @param SegmentNumber The number of the message segment.
+	 * @param SegmentOffset The segment's offset within the message.
+	 * @param SegmentData The segment data.
+	 * @param CurrentTime The current time.
 	 */
 	void Reassemble(int32 SegmentNumber, int32 SegmentOffset, const TArray<uint8>& SegmentData, const FDateTime& CurrentTime)
 	{
diff --git a/Engine/Plugins/Messaging/UdpMessaging/Source/UdpMessaging/Private/Transport/UdpMessageBeacon.cpp b/Engine/Plugins/Messaging/UdpMessaging/Source/UdpMessaging/Private/Transport/UdpMessageBeacon.cpp
--- a/Engine/Plugins/Messaging/UdpMessaging/Source/UdpMessaging/Private/Transport/UdpMessageBeacon.cpp
+++ b/Engine/Plugins/Messaging/UdpMessaging/Source/UdpMessaging/Private/Transport/UdpMessageBeacon.cpp
@@ -34,13 +34,13 @@ FUdpMessageBeacon::FUdpMessageBeacon(FSocket* InSocket, const FGuid& InSocketId,
 
 FUdpMessageBeacon::~FUdpMessageBeacon()
 {
-	if (Thread != NULL)
+	if (Thread != nullptr)
 	{
 		Thread->Kill(true);
 		delete Thread;
 	}
 
-	MulticastAddress = NULL;
+	MulticastAddress = nullptr;
 
 	delete EndpointLeftEvent;
 }
diff --git a/Engine/Plugins/Messaging/UdpMessaging/Source/UdpMessaging/Private/Transport/UdpMessageBeacon.h b/Engine/Plugins/Messaging/UdpMessaging/Source/UdpMessaging/Private/Transport/UdpMessageBeacon.h
--- a/Engine/Plugins/Messaging/UdpMessaging/Source/UdpMessaging/Private/Transport/UdpMessageBeacon.h
+++ b/Engine/Plugins/Messaging/UdpMessaging/Source/UdpMessaging/Private/Transport/UdpMessageBeacon.h
@@ -20,9 +20,9 @@ public:
 	/** 
 	 * Creates and initializes a new Hello sender.
 	 *
-	 * @param InSocket - The network socket used to send Hello segments.
-	 * @param InSocketId - The network socket identifier (used to detect unicast endpoint).
-	 * @param InMulticastEndpoint - The multicast group endpoint to transport messages to.
+	 * @param InSocket The network socket used to send Hello segments.
+	 * @param InSocketId The network socket identifier (used to detect unicast endpoint).
+	 * @param InMulticastEndpoint The multicast group endpoint to transport messages to.
 	 */
 	FUdpMessageBeacon(FSocket* InSocket, const FGuid& InSocketId, const FIPv4Endpoint& InMulticastEndpoint);
 
@@ -46,7 +46,7 @@ public:
 	/**
 	 * Sets the number of known IP endpoints.
 	 *
-	 * @param EndpointCount - The current number of known endpoints.
+	 * @param EndpointCount The current number of known endpoints.
 	 */
 	void SetEndpointCount(int32 EndpointCount);
 
@@ -69,7 +69,7 @@ protected:
 	/**
 	 * Sends the specified segment.
 	 *
-	 * @param SegmentType - The type of segment to send (Hello or Bye).
+	 * @param SegmentType The type of segment to send (Hello or Bye).
 	 */
 	void SendSegment(EUdpMessageSegments::Type SegmentType);
 
diff --git a/Engine/Plugins/Messaging/UdpMessaging/Source/UdpMessaging/Private/Transport/UdpMessageProcessor.cpp b/Engine/Plugins/Messaging/UdpMessaging/Source/UdpMessaging/Private/Transport/UdpMessageProcessor.cpp
--- a/Engine/Plugins/Messaging/UdpMessaging/Source/UdpMessaging/Private/Transport/UdpMessageProcessor.cpp
+++ b/Engine/Plugins/Messaging/UdpMessaging/Source/UdpMessaging/Private/Transport/UdpMessageProcessor.cpp
@@ -17,11 +17,11 @@ const int32 FUdpMessageProcessor::DeadHelloIntervals = 5;
  *****************************************************************************/
 
 FUdpMessageProcessor::FUdpMessageProcessor( FSocket* InSocket, const FGuid& InNodeId, const FIPv4Endpoint& InMulticastEndpoint )
-	: Beacon(NULL)
+	: Beacon(nullptr)
 	, LastSentMessage(-1)
 	, LocalNodeId(InNodeId)
 	, MulticastEndpoint(InMulticastEndpoint)
-	, Sender(NULL)
+	, Sender(nullptr)
 	, Socket(InSocket)
 	, Stopping(false)
 {
@@ -119,10 +119,10 @@ uint32 FUdpMessageProcessor::Run( )
 	}
 	
 	delete Beacon;
-	Beacon = NULL;
+	Beacon = nullptr;
 
 	delete Sender;
-	Sender = NULL;
+	Sender = nullptr;
 
 	return 0;
 }
@@ -345,7 +345,7 @@ void FUdpMessageProcessor::ProcessDataSegment( FInboundSegment& Segment, FNodeIn
 
 			if (NodeInfo.NodeId.IsValid())
 			{
-				MessageReceivedDelegate.ExecuteIfBound(HackReader, NULL, NodeInfo.NodeId);
+				MessageReceivedDelegate.ExecuteIfBound(HackReader, nullptr, NodeInfo.NodeId);
 			}
 		}
 		else if (NodeInfo.Resequencer.Resequence(Message))
@@ -357,7 +357,7 @@ void FUdpMessageProcessor::ProcessDataSegment( FInboundSegment& Segment, FNodeIn
 			{
 				if (NodeInfo.NodeId.IsValid())
 				{
-					MessageReceivedDelegate.ExecuteIfBound(HackReader, NULL, NodeInfo.NodeId);
+					MessageReceivedDelegate.ExecuteIfBound(HackReader, nullptr, NodeInfo.NodeId);
 				}
 			}
 		}
diff --git a/Engine/Plugins/Messaging/UdpMessaging/Source/UdpMessaging/Private/Transport/UdpMessageProcessor.h b/Engine/Plugins/Messaging/UdpMessaging/Source/UdpMessaging/Private/Transport/UdpMessageProcessor.h
--- a/Engine/Plugins/Messaging/UdpMessaging/Source/UdpMessaging/Private/Transport/UdpMessageProcessor.h
+++ b/Engine/Plugins/Messaging/UdpMessaging/Source/UdpMessaging/Private/Transport/UdpMessageProcessor.h
@@ -100,9 +100,9 @@ public:
 	/**
 	 * Creates and initializes a new message processor.
 	 *
-	 * @param InSocket - The network socket used to transport messages.
-	 * @param InNodeId - The local node identifier (used to detect the unicast endpoint).
-	 * @param InMulticastEndpoint - The multicast group endpoint to transport messages to.
+	 * @param InSocket The network socket used to transport messages.
+	 * @param InNodeId The local node identifier (used to detect the unicast endpoint).
+	 * @param InMulticastEndpoint The multicast group endpoint to transport messages to.
 	 */
 	FUdpMessageProcessor( FSocket* InSocket, const FGuid& InNodeId, const FIPv4Endpoint& InMulticastEndpoint );
 
@@ -116,8 +116,8 @@ public:
 	/**
 	 * Queues up an inbound message segment.
 	 *
-	 * @param Data - The segment data.
-	 * @param Sender - The sender's network endpoint.
+	 * @param Data The segment data.
+	 * @param Sender The sender's network endpoint.
 	 *
 	 * @return true if the segment was queued up, false otherwise.
 	 */
@@ -126,8 +126,8 @@ public:
 	/**
 	 * Queues up an outbound message.
 	 *
-	 * @param Data - The message data to send.
-	 * @param Recipient - The recipient's IPv4 endpoint.
+	 * @param Data The message data to send.
+	 * @param Recipient The recipient's IPv4 endpoint.
 	 *
 	 * @return true if the message was queued up, false otherwise.
 	 */
@@ -143,7 +143,7 @@ public:
 	/**
 	 * Returns a delegate that is executed when a channel has closed or timed out.
 	 *
-	 * @param Delegate - The delegate to set.
+	 * @param Delegate The delegate to set.
 	 */
 	FOnMessageTransportNodeLost& OnNodeLost( )
 	{
@@ -153,7 +153,7 @@ public:
 	/**
 	 * Returns a delegate that is executed when message data has been received.
 	 *
-	 * @param Delegate - The delegate to set.
+	 * @param Delegate The delegate to set.
 	 */
 	FOnMessageTransportMessageReceived& OnMessageReceived( )
 	{
@@ -179,8 +179,8 @@ protected:
 	/**
 	 * Acknowledges receipt of a message.
 	 *
-	 * @param MessageId - The identifier of the message to acknowledge.
-	 * @param NodeInfo - Details for the node to send the acknowledgment to.
+	 * @param MessageId The identifier of the message to acknowledge.
+	 * @param NodeInfo Details for the node to send the acknowledgment to.
 	 *
 	 * @todo gmp: batch multiple of these into a single message
 	 */
@@ -206,9 +206,9 @@ protected:
 	/**
 	 * Filters the specified message segment.
 	 *
-	 * @param Header - The segment header.
-	 * @param Data - The segment data.
-	 * @param Sender - The segment sender.
+	 * @param Header The segment header.
+	 * @param Data The segment data.
+	 * @param Sender The segment sender.
 	 *
 	 * @return true if the segment passed the filter, false otherwise.
 	 */
@@ -217,72 +217,72 @@ protected:
 	/**
 	 * Processes an Abort segment.
 	 *
-	 * @param Segment - The segment to process.
-	 * @param NodeInfo - Details for the node that sent the segment.
+	 * @param Segment The segment to process.
+	 * @param NodeInfo Details for the node that sent the segment.
 	 */
 	void ProcessAbortSegment( FInboundSegment& Segment, FNodeInfo& NodeInfo );
 
 	/**
 	 * Processes a Success segment.
 	 *
-	 * @param Segment - The segment to process.
-	 * @param NodeInfo - Details for the node that sent the segment.
+	 * @param Segment The segment to process.
+	 * @param NodeInfo Details for the node that sent the segment.
 	 */
 	void ProcessAcknowledgeSegment( FInboundSegment& Segment, FNodeInfo& NodeInfo );
 
 	/**
 	 * Processes a Bye segment.
 	 *
-	 * @param Segment - The segment to process.
-	 * @param NodeInfo - Details for the node that sent the segment.
+	 * @param Segment The segment to process.
+	 * @param NodeInfo Details for the node that sent the segment.
 	 */
 	void ProcessByeSegment( FInboundSegment& Segment, FNodeInfo& NodeInfo );
 
 	/**
 	 * Processes a Ping segment.
 	 *
-	 * @param Segment - The segment to process.
-	 * @param NodeInfo - Details for the node that sent the segment.
+	 * @param Segment The segment to process.
+	 * @param NodeInfo Details for the node that sent the segment.
 	 */
 	void ProcessDataSegment( FInboundSegment& Segment, FNodeInfo& NodeInfo );
 
 	/**
 	 * Processes a Hello segment.
 	 *
-	 * @param Segment - The segment to process.
-	 * @param NodeInfo - Details for the node that sent the segment.
+	 * @param Segment The segment to process.
+	 * @param NodeInfo Details for the node that sent the segment.
 	 */
 	void ProcessHelloSegment( FInboundSegment& Segment, FNodeInfo& NodeInfo );
 
 	/**
 	 * Processes a Ping segment.
 	 *
-	 * @param Segment - The segment to process.
-	 * @param NodeInfo - Details for the node that sent the segment.
+	 * @param Segment The segment to process.
+	 * @param NodeInfo Details for the node that sent the segment.
 	 */
 	void ProcessRetransmitSegment( FInboundSegment& Segment, FNodeInfo& NodeInfo );
 
 	/**
 	 * Processes a Timeout segment.
 	 *
-	 * @param Segment - The segment to process.
-	 * @param NodeInfo - Details for the node that sent the segment.
+	 * @param Segment The segment to process.
+	 * @param NodeInfo Details for the node that sent the segment.
 	 */
 	void ProcessTimeoutSegment( FInboundSegment& Segment, FNodeInfo& NodeInfo );
 
 	/**
 	 * Processes an unknown segment type.
 	 *
-	 * @param Segment - The segment to process.
-	 * @param NodeInfo - Details for the node that sent the segment.
-	 * @param SegmentType - The segment type.
+	 * @param Segment The segment to process.
+	 * @param NodeInfo Details for the node that sent the segment.
+	 * @param SegmentType The segment type.
 	 */
 	void ProcessUnknownSegment( FInboundSegment& Segment, FNodeInfo& NodeInfo, uint8 SegmentType );
 
 	/**
 	 * Removes the specified node from the list of known remote endpoints.
 	 *
-	 * @param NodeId - The identifier of the node to remove.
+	 * @param NodeId The identifier of the node to remove.
 	 */
 	void RemoveKnownNode( const FGuid& NodeId );
 
@@ -294,7 +294,7 @@ protected:
 	/**
 	 * Updates all segmenters of the specified node.
 	 *
-	 * @param NodeInfo - Details for the node to update.
+	 * @param NodeInfo Details for the node to update.
 	 */
 	void UpdateSegmenters( FNodeInfo& NodeInfo );
 
diff --git a/Engine/Plugins/Messaging/UdpMessaging/Source/UdpMessaging/Private/Transport/UdpMessageResequencer.h b/Engine/Plugins/Messaging/UdpMessaging/Source/UdpMessaging/Private/Transport/UdpMessageResequencer.h
--- a/Engine/Plugins/Messaging/UdpMessaging/Source/UdpMessaging/Private/Transport/UdpMessageResequencer.h
+++ b/Engine/Plugins/Messaging/UdpMessaging/Source/UdpMessaging/Private/Transport/UdpMessageResequencer.h
@@ -61,7 +61,7 @@ public:
 	/**
 	 * Resequences the specified message.
 	 *
-	 * @param Message - The message to resequence.
+	 * @param Message The message to resequence.
 	 *
 	 * @return true if the message is in sequence, false otherwise.
 	 */
diff --git a/Engine/Plugins/Messaging/UdpMessaging/Source/UdpMessaging/Private/Transport/UdpMessageSegmenter.h b/Engine/Plugins/Messaging/UdpMessaging/Source/UdpMessaging/Private/Transport/UdpMessageSegmenter.h
--- a/Engine/Plugins/Messaging/UdpMessaging/Source/UdpMessaging/Private/Transport/UdpMessageSegmenter.h
+++ b/Engine/Plugins/Messaging/UdpMessaging/Source/UdpMessaging/Private/Transport/UdpMessageSegmenter.h
@@ -21,17 +21,17 @@ public:
 	 * Default constructor.
 	 */
 	FUdpMessageSegmenter( )
-		: MessageReader(NULL)
+		: MessageReader(nullptr)
 	{ }
 
 	/**
 	 * Creates and initializes a new message segmenter.
 	 *
-	 * @param InMessage - The message to segment.
+	 * @param InMessage The message to segment.
 	 */
 	FUdpMessageSegmenter(const IMessageDataRef& InMessage, uint16 InSegmentSize)
 		: Message(InMessage)
-		, MessageReader(NULL)
+		, MessageReader(nullptr)
 		, SegmentSize(InSegmentSize)
 	{ }
 
@@ -40,7 +40,7 @@ public:
 	 */
 	~FUdpMessageSegmenter( )
 	{
-		if (MessageReader != NULL)
+		if (MessageReader != nullptr)
 		{
 			delete MessageReader;
 		}
@@ -55,7 +55,7 @@ public:
 	 */
 	int64 GetMessageSize() const
 	{
-		if (MessageReader == NULL)
+		if (MessageReader == nullptr)
 		{
 			return 0;
 		}
@@ -66,14 +66,14 @@ public:
 	/**
 	 * Gets the next pending segment.
 	 *
-	 * @param OutData - Will hold the segment data.
-	 * @param OutSegment - Will hold the segment number.
+	 * @param OutData Will hold the segment data.
+	 * @param OutSegment Will hold the segment number.
 	 *
 	 * @return true if a segment was returned, false if there are no more pending segments.
 	 */
 	bool GetNextPendingSegment(TArray<uint8>& OutData, uint16& OutSegment) const
 	{
-		if (MessageReader == NULL)
+		if (MessageReader == nullptr)
 		{
 			return false;
 		}
@@ -129,7 +129,7 @@ public:
 	 */
 	void Initialize()
 	{
-		if (MessageReader != NULL)
+		if (MessageReader != nullptr)
 		{
 			return;
 		}
@@ -159,7 +159,7 @@ public:
 	 */
 	bool IsInitialized() const
 	{
-		return (MessageReader != NULL);
+		return (MessageReader != nullptr);
 	}
 
 	/**
@@ -175,7 +175,7 @@ public:
 	/**
 	 * Marks the specified segment as sent.
 	 *
-	 * @param Segment - The sent segment.
+	 * @param Segment The sent segment.
 	 */
 	void MarkAsSent(uint16 Segment)
 	{
@@ -197,7 +197,7 @@ public:
 	/**
 	 * Marks the specified segments for retransmission.
 	 *
-	 * @param Segments - The data segments to retransmit.
+	 * @param Segments The data segments to retransmit.
 	 */
 	void MarkForRetransmission(const TArray<uint16>& Segments)
 	{
diff --git a/Engine/Plugins/Messaging/UdpMessaging/Source/UdpMessaging/Private/Transport/UdpMessageTransport.h b/Engine/Plugins/Messaging/UdpMessaging/Source/UdpMessaging/Private/Transport/UdpMessageTransport.h
--- a/Engine/Plugins/Messaging/UdpMessaging/Source/UdpMessaging/Private/Transport/UdpMessageTransport.h
+++ b/Engine/Plugins/Messaging/UdpMessaging/Source/UdpMessaging/Private/Transport/UdpMessageTransport.h
@@ -18,9 +18,9 @@ public:
 	/**
 	 * Creates and initializes a new instance.
 	 *
-	 * @param InLocalEndpoint - The local IP endpoint to receive messages on.
-	 * @param InMulticastEndpoint - The multicast group endpoint to transport messages to.
-	 * @param InMulticastTtl - The multicast time-to-live.
+	 * @param InLocalEndpoint The local IP endpoint to receive messages on.
+	 * @param InMulticastEndpoint The multicast group endpoint to transport messages to.
+	 * @param InMulticastTtl The multicast time-to-live.
 	 */
 	FUdpMessageTransport( const FIPv4Endpoint& InLocalEndpoint, const FIPv4Endpoint& InMulticastEndpoint, uint8 InMulticastTtl );
 
diff --git a/Engine/Plugins/Messaging/UdpMessaging/Source/UdpMessaging/Private/Tunnel/UdpMessageTunnel.h b/Engine/Plugins/Messaging/UdpMessaging/Source/UdpMessaging/Private/Tunnel/UdpMessageTunnel.h
--- a/Engine/Plugins/Messaging/UdpMessaging/Source/UdpMessaging/Private/Tunnel/UdpMessageTunnel.h
+++ b/Engine/Plugins/Messaging/UdpMessaging/Source/UdpMessaging/Private/Tunnel/UdpMessageTunnel.h
@@ -32,11 +32,11 @@ public:
 	/**
 	 * Creates and initializes a new instance.
 	 *
-	 * @param InLocalEndpoint - The local IP endpoint to receive messages on.
-	 * @param InMulticastEndpoint - The multicast group endpoint to transport messages to.
+	 * @param InLocalEndpoint The local IP endpoint to receive messages on.
+	 * @param InMulticastEndpoint The multicast group endpoint to transport messages to.
 	 */
 	FUdpMessageTunnel( const FIPv4Endpoint& InUnicastEndpoint, const FIPv4Endpoint& InMulticastEndpoint )
-		: Listener(NULL)
+		: Listener(nullptr)
 		, MulticastEndpoint(InMulticastEndpoint)
 		, Stopping(false)
 		, TotalInboundBytes(0)
@@ -68,23 +68,23 @@ public:
 	 */
 	~FUdpMessageTunnel( )
 	{
-		if (Thread != NULL)
+		if (Thread != nullptr)
 		{
 			Thread->Kill(true);
 			delete Thread;
 		}
 
 		// destroy sockets
-		if (MulticastSocket != NULL)
+		if (MulticastSocket != nullptr)
 		{
 			ISocketSubsystem::Get(PLATFORM_SOCKETSUBSYSTEM)->DestroySocket(MulticastSocket);
-			MulticastSocket = NULL;
+			MulticastSocket = nullptr;
 		}
 	
-		if (UnicastSocket != NULL)
+		if (UnicastSocket != nullptr)
 		{
 			ISocketSubsystem::Get(PLATFORM_SOCKETSUBSYSTEM)->DestroySocket(UnicastSocket);
-			UnicastSocket = NULL;
+			UnicastSocket = nullptr;
 		}
 	}
 
@@ -134,7 +134,7 @@ public:
 	{
 		FSocket* Socket = FTcpSocketBuilder(TEXT("FUdpMessageTunnel.RemoteConnection"));
 
-		if (Socket != NULL)
+		if (Socket != nullptr)
 		{
 			if (Socket->Connect(*RemoteEndpoint.ToInternetAddr()))
 			{
@@ -175,7 +175,7 @@ public:
 
 	virtual bool IsServerRunning( ) const OVERRIDE
 	{
-		return (Listener != NULL);
+		return (Listener != nullptr);
 	}
 
 	virtual FSimpleDelegate& OnConnectionsChanged( ) OVERRIDE
@@ -194,7 +194,7 @@ public:
 	virtual void StopServer( ) OVERRIDE
 	{
 		delete Listener;
-		Listener = NULL;
+		Listener = nullptr;
 	}
 
 	// End IUdpMessageTunnel interface
@@ -204,7 +204,7 @@ protected:
 	/**
 	 * Removes expired nodes from the specified collection of node infos.
 	 *
-	 * @param Nodes - The collection of nodes to clean up.
+	 * @param Nodes The collection of nodes to clean up.
 	 */
 	void RemoveExpiredNodes( TMap<FGuid, FNodeInfo>& Nodes )
 	{
@@ -254,7 +254,7 @@ protected:
 				{
 					FNodeInfo* LocalNode = LocalNodes.Find(Header.RecipientNodeId);
 
-					if (LocalNode == NULL)
+					if (LocalNode == nullptr)
 					{
 						continue;
 					}
@@ -279,7 +279,7 @@ protected:
 	/**
 	 * Receives all buffered datagrams from the specified socket and forwards them to the tunnels.
 	 *
-	 * @param Socket - The socket to receive from.
+	 * @param Socket The socket to receive from.
 	 */
 	void UdpToTcp( FSocket* Socket )
 	{
@@ -318,7 +318,7 @@ protected:
 				{
 					FNodeInfo* RemoteNode = RemoteNodes.Find(Header.RecipientNodeId);
 
-					if ((RemoteNode != NULL) && (RemoteNode->Connection.IsValid()))
+					if ((RemoteNode != nullptr) && (RemoteNode->Connection.IsValid()))
 					{
 						RemoteNode->Connection->Send(Datagram);
 					}
diff --git a/Engine/Plugins/Messaging/UdpMessaging/Source/UdpMessaging/Private/Tunnel/UdpMessageTunnelConnection.h b/Engine/Plugins/Messaging/UdpMessaging/Source/UdpMessaging/Private/Tunnel/UdpMessageTunnelConnection.h
--- a/Engine/Plugins/Messaging/UdpMessaging/Source/UdpMessaging/Private/Tunnel/UdpMessageTunnelConnection.h
+++ b/Engine/Plugins/Messaging/UdpMessaging/Source/UdpMessaging/Private/Tunnel/UdpMessageTunnelConnection.h
@@ -23,12 +23,6 @@ typedef TSharedRef<class FUdpMessageTunnelConnection> FUdpMessageTunnelConnectio
 typedef TSharedRef<TArray<uint8>, ESPMode::ThreadSafe> FSaveByteArrayRef;
 
 
-struct FUdpMessageTunnelPacket
-{
-
-};
-
-
 /**
  * Delegate type for received packets.
  *
@@ -51,8 +45,8 @@ public:
 	/**
 	 * Creates and initializes a new instance.
 	 *
-	 * @param InSocket - The socket to use for this connection.
-	 * @param InRemoteAddress - The IP endpoint of the remote client.
+	 * @param InSocket The socket to use for this connection.
+	 * @param InRemoteAddress The IP endpoint of the remote client.
 	 */
 	FUdpMessageTunnelConnection( FSocket* InSocket, const FIPv4Endpoint& InRemoteEndpoint )
 		: OpenedTime(FDateTime::UtcNow())
@@ -73,14 +67,14 @@ public:
 	 */
 	~FUdpMessageTunnelConnection( )
 	{
-		if (Thread != NULL)
+		if (Thread != nullptr)
 		{
 			Thread->Kill(true);
 			delete Thread;
 		}
 
 		delete Socket;
-		Socket = NULL;
+		Socket = nullptr;
 	}
 
 public:
@@ -88,7 +82,7 @@ public:
 	/**
 	 * Receives a payload from the connection's inbox.
 	 *
-	 * @param OutPayload - Will hold the received payload, if any.
+	 * @param OutPayload Will hold the received payload, if any.
 	 *
 	 * @return true if a payload was returned, false otherwise.
 	 */
@@ -100,7 +94,7 @@ public:
 	/**
 	 * Sends a payload through this connection.
 	 *
-	 * @param Payload - The payload to send.
+	 * @param Payload The payload to send.
 	 */
 	bool Send( const FArrayReaderPtr& Payload )
 	{
@@ -118,7 +112,7 @@ public:
 
 	virtual bool Init( ) OVERRIDE
 	{
-		return (Socket != NULL);
+		return (Socket != nullptr);
 	}
 
 	virtual uint32 Run( ) OVERRIDE
diff --git a/Engine/Plugins/Messaging/UdpMessaging/Source/UdpMessaging/Private/UdpMessagingPrivatePCH.h b/Engine/Plugins/Messaging/UdpMessaging/Source/UdpMessaging/Private/UdpMessagingPrivatePCH.h
--- a/Engine/Plugins/Messaging/UdpMessaging/Source/UdpMessaging/Private/UdpMessagingPrivatePCH.h
+++ b/Engine/Plugins/Messaging/UdpMessaging/Source/UdpMessaging/Private/UdpMessagingPrivatePCH.h
@@ -9,7 +9,7 @@
 #include "UdpMessaging.h"
 
 
-/* Dependencies
+/* Private Dependencies
  *****************************************************************************/
 
 #include "Core.h"
@@ -21,7 +21,7 @@
 #include "SocketSubsystem.h"
 
 
-/* Constants
+/* Private Constants
  *****************************************************************************/
 
 /**
diff --git a/Engine/Plugins/Messaging/UdpMessaging/Source/UdpMessaging/Public/Interfaces/IUdpMessageTunnel.h b/Engine/Plugins/Messaging/UdpMessaging/Source/UdpMessaging/Public/Interfaces/IUdpMessageTunnel.h
--- a/Engine/Plugins/Messaging/UdpMessaging/Source/UdpMessaging/Public/Interfaces/IUdpMessageTunnel.h
+++ b/Engine/Plugins/Messaging/UdpMessaging/Source/UdpMessaging/Public/Interfaces/IUdpMessageTunnel.h
@@ -28,7 +28,7 @@ public:
 	/**
 	 * Establishes a tunnel with a remote server.
 	 *
-	 * @param RemoteEndpoint - The endpoint of the server to connect to.
+	 * @param RemoteEndpoint The endpoint of the server to connect to.
 	 *
 	 * @return true if the connection has been established, false otherwise.
 	 *
@@ -39,7 +39,7 @@ public:
 	/** 
 	 * Starts the tunnel server.
 	 *
-	 * @param LocalEndpoint - The IP endpoint to listen for incoming connections on.
+	 * @param LocalEndpoint The IP endpoint to listen for incoming connections on.
 	 *
 	 * @see IsServerRunning
 	 * @see StopServer
@@ -59,7 +59,7 @@ public:
 	/**
 	 * Gets the list of all open tunnel connections.
 	 *
-	 * @param OutConnections - Will hold the list of connections.
+	 * @param OutConnections Will hold the list of connections.
 	 *
 	 * @return The number of connections returned.
 	 *
diff --git a/Engine/Plugins/Messaging/UdpMessaging/Source/UdpMessaging/Public/UdpMessaging.h b/Engine/Plugins/Messaging/UdpMessaging/Source/UdpMessaging/Public/UdpMessaging.h
--- a/Engine/Plugins/Messaging/UdpMessaging/Source/UdpMessaging/Public/UdpMessaging.h
+++ b/Engine/Plugins/Messaging/UdpMessaging/Source/UdpMessaging/Public/UdpMessaging.h
@@ -7,13 +7,13 @@
 #pragma once
 
 
-/* Dependencies
+/* Public Dependencies
  *****************************************************************************/
 
 #include "Networking.h"
 
 
-/* Interfaces
+/* Public Interfaces
  *****************************************************************************/
 
 #include "IUdpMessageTunnelConnection.h"
