UdpMessaging: endpoint string conversion cleanup

diff --git a/Engine/Plugins/Messaging/UdpMessaging/Source/UdpMessaging/Private/UdpMessagingModule.cpp b/Engine/Plugins/Messaging/UdpMessaging/Source/UdpMessaging/Private/UdpMessagingModule.cpp
--- a/Engine/Plugins/Messaging/UdpMessaging/Source/UdpMessaging/Private/UdpMessagingModule.cpp
+++ b/Engine/Plugins/Messaging/UdpMessaging/Source/UdpMessaging/Private/UdpMessagingModule.cpp
@@ -158,7 +158,7 @@ protected:
 		{
 			GLog->Logf(TEXT("Warning: Invalid UDP Tunneling UnicastEndpoint '%s' - binding to all local network adapters instead"), *Settings->UnicastEndpoint);
 			UnicastEndpoint = FIPv4Endpoint::Any;
-			Settings->UnicastEndpoint = UnicastEndpoint.ToText().ToString();
+			Settings->UnicastEndpoint = UnicastEndpoint.ToString();
 			ResaveSettings = true;
 		}
 
@@ -166,7 +166,7 @@ protected:
 		{
 			GLog->Logf(TEXT("Warning: Invalid UDP Tunneling MulticastEndpoint '%s' - using default endpoint '%s' instead"), *Settings->MulticastEndpoint, *UDP_MESSAGING_DEFAULT_MULTICAST_ENDPOINT.ToText().ToString());
 			MulticastEndpoint = UDP_MESSAGING_DEFAULT_MULTICAST_ENDPOINT;
-			Settings->MulticastEndpoint = MulticastEndpoint.ToText().ToString();
+			Settings->MulticastEndpoint = MulticastEndpoint.ToString();
 			ResaveSettings = true;
 		}
 
@@ -238,13 +238,13 @@ protected:
 	{
 		if (FApp::IsGame())
 		{
-			// only allow UDP transport in Debug and Development configurations
+			// disallow in Debug and Development configurations
 			if ((FApp::GetBuildConfiguration() == EBuildConfigurations::Shipping) || (FApp::GetBuildConfiguration() == EBuildConfigurations::Test))
 			{
 				return false;
 			}
 
-			// only allow UDP transport if the platform supports it
+			// disallow unsupported platforms
 			if (!FPlatformMisc::SupportsMessaging() || !FPlatformProcess::SupportsMultithreading())
 			{
 				return false;
