UdpMessaging: Documentation cleanup pass

diff --git a/Engine/Plugins/Messaging/UdpMessaging/Source/UdpMessaging/Private/Shared/UdpMessageSegment.h b/Engine/Plugins/Messaging/UdpMessaging/Source/UdpMessaging/Private/Shared/UdpMessageSegment.h
--- a/Engine/Plugins/Messaging/UdpMessaging/Source/UdpMessaging/Private/Shared/UdpMessageSegment.h
+++ b/Engine/Plugins/Messaging/UdpMessaging/Source/UdpMessaging/Private/Shared/UdpMessageSegment.h
@@ -10,44 +10,28 @@ namespace EUdpMessageSegments
 	 */
 	enum Type
 	{
-		/**
-		 * None.
-		 */
+		/** None. */
 		None,
 
-		/**
-		 * Request to abort the sending of a message.
-		 */
+		/** Request to abort the sending of a message. */
 		Abort,
 
-		/**
-		 * Acknowledges that the message was received successfully.
-		 */
+		/** Acknowledges that the message was received successfully. */
 		Acknowledge,
 
-		/**
-		 * Notifies the bus that an endpoint has left.
-		 */
+		/** Notifies the bus that an endpoint has left. */
 		Bye,
 
-		/**
-		 * A message data segment.
-		 */
+		/** A message data segment. */
 		Data,
 
-		/**
-		 * Notifies the bus that an endpoint has joined.
-		 */
+		/** Notifies the bus that an endpoint has joined. */
 		Hello,
 
-		/**
-		 * Request to retransmit selected data segments.
-		 */
+		/** Request to retransmit selected data segments. */
 		Retransmit,
 
-		/**
-		 * Notification that an inbound message timed out.
-		 */
+		/** Notification that an inbound message timed out. */
 		Timeout
 	};
 }
@@ -60,19 +44,13 @@ namespace FUdpMessageSegment
 	 */
 	struct FHeader
 	{
-		/**
-		 * Holds the protocol version.
-		 */
+		/** Holds the protocol version. */
 		uint8 ProtocolVersion;
 
-		/**
-		 * Holds the recipient's node identifier (empty = multicast).
-		 */
+		/** Holds the recipient's node identifier (empty = multicast). */
 		FGuid RecipientNodeId;
 
-		/**
-		 * Holds the sender's node identifier.
-		 */
+		/** Holds the sender's node identifier. */
 		FGuid SenderNodeId;
 
 		/**
@@ -89,10 +67,9 @@ namespace FUdpMessageSegment
 		 *
 		 * @param Ar The archive to serialize from or into.
 		 * @param DateTime The header to serialize.
-		 *
 		 * @return The archive.
 		 */
-		friend FArchive& operator<< (FArchive& Ar, FHeader& Header)
+		friend FArchive& operator<<( FArchive& Ar, FHeader& Header )
 		{
 			return Ar << Header.ProtocolVersion << Header.RecipientNodeId << Header.SenderNodeId << Header.SegmentType;
 		}
@@ -104,9 +81,7 @@ namespace FUdpMessageSegment
 	 */
 	struct FAbortChunk
 	{
-		/** 
-		 * Holds the identifier of the message to abort.
-		 */
+		/** Holds the identifier of the message to abort. */
 		int32 MessageId;
 
 	public:
@@ -116,10 +91,9 @@ namespace FUdpMessageSegment
 		 *
 		 * @param Ar The archive to serialize from or into.
 		 * @param DateTime The header to serialize.
-		 *
 		 * @return The archive.
 		 */
-		friend FArchive& operator<< (FArchive& Ar, FAbortChunk& Header)
+		friend FArchive& operator<<( FArchive& Ar, FAbortChunk& Header )
 		{
 			return Ar << Header.MessageId;
 		}
@@ -131,9 +105,7 @@ namespace FUdpMessageSegment
 	 */
 	struct FAcknowledgeChunk
 	{
-		/** 
-		 * Holds the identifier of the message that was received successfully.
-		 */
+		/** Holds the identifier of the message that was received successfully. */
 		int32 MessageId;
 
 	public:
@@ -143,10 +115,9 @@ namespace FUdpMessageSegment
 		 *
 		 * @param Ar The archive to serialize from or into.
 		 * @param DateTime The header to serialize.
-		 *
 		 * @return The archive.
 		 */
-		friend FArchive& operator<< (FArchive& Ar, FAcknowledgeChunk& Header)
+		friend FArchive& operator<<( FArchive& Ar, FAcknowledgeChunk& Header )
 		{
 			return Ar << Header.MessageId;
 		}
@@ -158,39 +129,25 @@ namespace FUdpMessageSegment
 	 */
 	struct FDataChunk
 	{
-		/**
-		 * Holds the identifier of the message that the data belongs to.
-		 */
+		/** Holds the identifier of the message that the data belongs to. */
 		int32 MessageId;
 
-		/**
-		 * Holds the total size of the message.
-		 */
+		/** Holds the total size of the message. */
 		int32 MessageSize;
 
-		/**
-		 * Holds the sequence number of this segment.
-		 */
+		/** Holds the sequence number of this segment. */
 		uint16 SegmentNumber;
 
-		/**
-		 * Holds the segment's offset within the message.
-		 */
+		/** Holds the segment's offset within the message. */
 		uint32 SegmentOffset;
 
-		/**
-		 * Holds the message sequence number (0 = not sequential).
-		 */
+		/** Holds the message sequence number (0 = not sequential). */
 		uint64 Sequence;
 
-		/**
-		 * Holds the total number of data segments being sent.
-		 */
+		/** Holds the total number of data segments being sent. */
 		uint16 TotalSegments;
 
-		/**
-		 * Holds the segment data.
-		 */
+		/** Holds the segment data. */
 		TArray<uint8> Data;
 
 	public:
@@ -200,10 +157,9 @@ namespace FUdpMessageSegment
 		 *
 		 * @param Ar The archive to serialize from or into.
 		 * @param DateTime The header to serialize.
-		 *
 		 * @return The archive.
 		 */
-		friend FArchive& operator<< (FArchive& Ar, FDataChunk& Chunk)
+		friend FArchive& operator<<( FArchive& Ar, FDataChunk& Chunk )
 		{
 			return Ar
 				<< Chunk.MessageId
@@ -226,14 +182,10 @@ namespace FUdpMessageSegment
 	 */
 	struct FRetransmitChunk
 	{
-		/**
-		 * Holds the identifier of the message for which data needs to be retransmitted.
-		 */
+		/** Holds the identifier of the message for which data needs to be retransmitted. */
 		int32 MessageId;
 
-		/**
-		 * Holds the list of data segments that need to be retransmitted.
-		 */
+		/** Holds the list of data segments that need to be retransmitted. */
 		TArray<uint16> Segments;
 
 	public:
@@ -243,10 +195,9 @@ namespace FUdpMessageSegment
 		 *
 		 * @param Ar The archive to serialize from or into.
 		 * @param DateTime The header to serialize.
-		 *
 		 * @return The archive.
 		 */
-		friend FArchive& operator<< (FArchive& Ar, FRetransmitChunk& Header)
+		friend FArchive& operator<<( FArchive& Ar, FRetransmitChunk& Header )
 		{
 			return Ar << Header.MessageId << Header.Segments;
 		}
@@ -258,9 +209,7 @@ namespace FUdpMessageSegment
 	 */
 	struct FTimeoutChunk
 	{
-		/** 
-		 * Holds the identifier of the message that timed out.
-		 */
+		/** Holds the identifier of the message that timed out. */
 		int32 MessageId;
 
 	public:
@@ -270,12 +219,11 @@ namespace FUdpMessageSegment
 		 *
 		 * @param Ar The archive to serialize from or into.
 		 * @param DateTime The header to serialize.
-		 *
 		 * @return The archive.
 		 */
-		friend FArchive& operator<< (FArchive& Ar, FTimeoutChunk& Header)
+		friend FArchive& operator<<( FArchive& Ar, FTimeoutChunk& Header )
 		{
 			return Ar << Header.MessageId;
 		}
 	};
-};
\ No newline at end of file
+};
diff --git a/Engine/Plugins/Messaging/UdpMessaging/Source/UdpMessaging/Private/Shared/UdpMessagingSettings.h b/Engine/Plugins/Messaging/UdpMessaging/Source/UdpMessaging/Private/Shared/UdpMessagingSettings.h
--- a/Engine/Plugins/Messaging/UdpMessaging/Source/UdpMessaging/Private/Shared/UdpMessagingSettings.h
+++ b/Engine/Plugins/Messaging/UdpMessaging/Source/UdpMessaging/Private/Shared/UdpMessagingSettings.h
@@ -17,27 +17,19 @@ class UUdpMessagingSettings
 
 public:
 
-	/**
-	 * Whether the UDP transport channel is enabled.
-	 */
+	/** Whether the UDP transport channel is enabled. */
 	UPROPERTY(config, EditAnywhere, Category=Transport)
 	bool EnableTransport;
 
-	/**
-	 * The IP endpoint to listen to and send packets from.
-	 */
+	/** The IP endpoint to listen to and send packets from. */
 	UPROPERTY(config, EditAnywhere, Category=Transport)
 	FString UnicastEndpoint;
 
-	/**
-	 * The IP endpoint to send multicast packets to.
-	 */
+	/** The IP endpoint to send multicast packets to. */
 	UPROPERTY(config, EditAnywhere, Category=Transport)
 	FString MulticastEndpoint;
 
-	/**
-	 * The time-to-live (TTL) for sent multicast packets.
-	 */
+	/** The time-to-live (TTL) for sent multicast packets. */
 	UPROPERTY(config, EditAnywhere, Category=Transport)
 	uint8 MulticastTimeToLive;
 
@@ -51,21 +43,15 @@ public:
 
 public:
 
-	/**
-	 * Whether the UDP tunnel is enabled.
-	 */
+	/** Whether the UDP tunnel is enabled. */
 	UPROPERTY(config, EditAnywhere, Category=Tunnel)
 	bool EnableTunnel;
 
-	/**
-	 * The IP endpoint to listen to and send packets from.
-	 */
+	/** The IP endpoint to listen to and send packets from. */
 	UPROPERTY(config, EditAnywhere, Category=Tunnel)
 	FString TunnelUnicastEndpoint;
 
-	/**
-	 * The IP endpoint to send multicast packets to.
-	 */
+	/** The IP endpoint to send multicast packets to. */
 	UPROPERTY(config, EditAnywhere, Category=Tunnel)
 	FString TunnelMulticastEndpoint;
 
diff --git a/Engine/Plugins/Messaging/UdpMessaging/Source/UdpMessaging/Private/Transport/ReassembledUdpMessage.h b/Engine/Plugins/Messaging/UdpMessaging/Source/UdpMessaging/Private/Transport/ReassembledUdpMessage.h
--- a/Engine/Plugins/Messaging/UdpMessaging/Source/UdpMessaging/Private/Transport/ReassembledUdpMessage.h
+++ b/Engine/Plugins/Messaging/UdpMessaging/Source/UdpMessaging/Private/Transport/ReassembledUdpMessage.h
@@ -7,6 +7,13 @@
 #pragma once
 
 
+/** Type definition for shared pointers to instances of FUdpMessage. */
+typedef TSharedPtr<class FReassembledUdpMessage, ESPMode::ThreadSafe> FReassembledUdpMessagePtr;
+
+/** Type definition for shared references to instances of FUdpMessage. */
+typedef TSharedRef<class FReassembledUdpMessage, ESPMode::ThreadSafe> FReassembledUdpMessageRef;
+
+
 /**
  * Implements a reassembled message.
  */
@@ -17,7 +24,7 @@ public:
 	/**
 	 * Default constructor.
 	 */
-	FReassembledUdpMessage() { }
+	FReassembledUdpMessage( ) { }
 
 	/**
 	 * Creates and initializes a new inbound message info.
@@ -27,7 +34,7 @@ public:
 	 * @param InSequence The message sequence number.
 	 * @param InSender The IPv4 endpoint of the sender.
 	 */
-	FReassembledUdpMessage(int32 MessageSize, int32 SegmentCount, uint64 InSequence, const FIPv4Endpoint& InSender)
+	FReassembledUdpMessage( int32 MessageSize, int32 SegmentCount, uint64 InSequence, const FIPv4Endpoint& InSender )
 		: PendingSegments(true, SegmentCount)
 		, PendingSegmentsCount(SegmentCount)
 		, ReceivedBytes(0)
@@ -49,7 +56,7 @@ public:
 	 *
 	 * @return Last receive time.
 	 */
-	FDateTime GetLastSegmentTime() const
+	FDateTime GetLastSegmentTime( ) const
 	{
 		return LastSegmentTime;
 	}
@@ -59,7 +66,7 @@ public:
 	 *
 	 * @return List of pending segment numbers.
 	 */
-	TArray<uint16> GetPendingSegments() const
+	TArray<uint16> GetPendingSegments( ) const
 	{
 		TArray<uint16> Result;
 
@@ -79,7 +86,7 @@ public:
 	 *
 	 * @return Number of pending segments.
 	 */
-	uint16 GetPendingSegmentsCount() const
+	uint16 GetPendingSegmentsCount( ) const
 	{
 		return PendingSegmentsCount;
 	}
@@ -89,7 +96,7 @@ public:
 	 *
 	 * @return Number of retransmit requests.
 	 */
-	int32 GetRetransmitRequestsCount() const
+	int32 GetRetransmitRequestsCount( ) const
 	{
 		return RetransmitRequestsCount;
 	}
@@ -99,7 +106,7 @@ public:
 	 *
 	 * @return true if the message is complete, false otherwise.
 	 */
-	bool IsComplete() const
+	bool IsComplete( ) const
 	{
 		return (PendingSegmentsCount == 0);
 	}
@@ -109,7 +116,7 @@ public:
 	 *
 	 * @return true if the message is initialized, false otherwise.
 	 */
-	bool IsInitialized() const
+	bool IsInitialized( ) const
 	{
 		return (Data.Num() < 0);
 	}
@@ -122,7 +129,7 @@ public:
 	 * @param SegmentData The segment data.
 	 * @param CurrentTime The current time.
 	 */
-	void Reassemble(int32 SegmentNumber, int32 SegmentOffset, const TArray<uint8>& SegmentData, const FDateTime& CurrentTime)
+	void Reassemble( int32 SegmentNumber, int32 SegmentOffset, const TArray<uint8>& SegmentData, const FDateTime& CurrentTime )
 	{
 		if (SegmentNumber >= PendingSegments.Num())
 		{
@@ -153,7 +160,7 @@ public:
 	 *
 	 * @return Message data.
 	 */
-	virtual const TArray<uint8>& GetData() const
+	virtual const TArray<uint8>& GetData( ) const
 	{
 		return Data;
 	}
@@ -163,7 +170,7 @@ public:
 	 *
 	 * @return Sequence number.
 	 */
-	uint64 GetSequence() const
+	uint64 GetSequence( ) const
 	{
 		return Sequence;
 	}
@@ -194,14 +201,3 @@ private:
 	// Holds the message sequence.
 	uint64 Sequence;
 };
-
-
-/**
- * Type definition for shared pointers to instances of FUdpMessage.
- */
-typedef TSharedPtr<FReassembledUdpMessage, ESPMode::ThreadSafe> FReassembledUdpMessagePtr;
-
-/**
- * Type definition for shared references to instances of FUdpMessage.
- */
-typedef TSharedRef<FReassembledUdpMessage, ESPMode::ThreadSafe> FReassembledUdpMessageRef;
diff --git a/Engine/Plugins/Messaging/UdpMessaging/Source/UdpMessaging/Private/Transport/UdpMessageProcessor.h b/Engine/Plugins/Messaging/UdpMessaging/Source/UdpMessaging/Private/Transport/UdpMessageProcessor.h
--- a/Engine/Plugins/Messaging/UdpMessaging/Source/UdpMessaging/Private/Transport/UdpMessageProcessor.h
+++ b/Engine/Plugins/Messaging/UdpMessaging/Source/UdpMessaging/Private/Transport/UdpMessageProcessor.h
@@ -118,7 +118,6 @@ public:
 	 *
 	 * @param Data The segment data.
 	 * @param Sender The sender's network endpoint.
-	 *
 	 * @return true if the segment was queued up, false otherwise.
 	 */
 	bool EnqueueInboundSegment( const FArrayReaderPtr& Data, const FIPv4Endpoint& Sender );
@@ -128,7 +127,6 @@ public:
 	 *
 	 * @param Data The message data to send.
 	 * @param Recipient The recipient's IPv4 endpoint.
-	 *
 	 * @return true if the message was queued up, false otherwise.
 	 */
 	bool EnqueueOutboundMessage( const IMessageDataRef& Data, const FGuid& Recipient );
@@ -209,7 +207,6 @@ protected:
 	 * @param Header The segment header.
 	 * @param Data The segment data.
 	 * @param Sender The segment sender.
-	 *
 	 * @return true if the segment passed the filter, false otherwise.
 	 */
 	bool FilterSegment( const FUdpMessageSegment::FHeader& Header, const FArrayReaderPtr& Data, const FIPv4Endpoint& Sender );
diff --git a/Engine/Plugins/Messaging/UdpMessaging/Source/UdpMessaging/Private/Transport/UdpMessageResequencer.h b/Engine/Plugins/Messaging/UdpMessaging/Source/UdpMessaging/Private/Transport/UdpMessageResequencer.h
--- a/Engine/Plugins/Messaging/UdpMessaging/Source/UdpMessaging/Private/Transport/UdpMessageResequencer.h
+++ b/Engine/Plugins/Messaging/UdpMessaging/Source/UdpMessaging/Private/Transport/UdpMessageResequencer.h
@@ -62,7 +62,6 @@ public:
 	 * Resequences the specified message.
 	 *
 	 * @param Message The message to resequence.
-	 *
 	 * @return true if the message is in sequence, false otherwise.
 	 */
 	bool Resequence(const FReassembledUdpMessagePtr& Message)
diff --git a/Engine/Plugins/Messaging/UdpMessaging/Source/UdpMessaging/Private/Transport/UdpMessageSegmenter.h b/Engine/Plugins/Messaging/UdpMessaging/Source/UdpMessaging/Private/Transport/UdpMessageSegmenter.h
--- a/Engine/Plugins/Messaging/UdpMessaging/Source/UdpMessaging/Private/Transport/UdpMessageSegmenter.h
+++ b/Engine/Plugins/Messaging/UdpMessaging/Source/UdpMessaging/Private/Transport/UdpMessageSegmenter.h
@@ -68,7 +68,6 @@ public:
 	 *
 	 * @param OutData Will hold the segment data.
 	 * @param OutSegment Will hold the segment number.
-	 *
 	 * @return true if a segment was returned, false if there are no more pending segments.
 	 */
 	bool GetNextPendingSegment(TArray<uint8>& OutData, uint16& OutSegment) const
diff --git a/Engine/Plugins/Messaging/UdpMessaging/Source/UdpMessaging/Private/Tunnel/UdpMessageTunnelConnection.h b/Engine/Plugins/Messaging/UdpMessaging/Source/UdpMessaging/Private/Tunnel/UdpMessageTunnelConnection.h
--- a/Engine/Plugins/Messaging/UdpMessaging/Source/UdpMessaging/Private/Tunnel/UdpMessageTunnelConnection.h
+++ b/Engine/Plugins/Messaging/UdpMessaging/Source/UdpMessaging/Private/Tunnel/UdpMessageTunnelConnection.h
@@ -7,19 +7,13 @@
 #pragma once
 
 
-/**
- * Type definition for shared pointers to instances of FUdpMessageTunnelConnection.
- */
+/** Type definition for shared pointers to instances of FUdpMessageTunnelConnection. */
 typedef TSharedPtr<class FUdpMessageTunnelConnection> FUdpMessageTunnelConnectionPtr;
 
-/**
- * Type definition for shared references to instances of FUdpMessageTunnelConnection.
- */
+/** Type definition for shared references to instances of FUdpMessageTunnelConnection. */
 typedef TSharedRef<class FUdpMessageTunnelConnection> FUdpMessageTunnelConnectionRef;
 
-/**
- * Type definition for thread-safe shared references to byte arrays.
- */
+/** Type definition for thread-safe shared references to byte arrays. */
 typedef TSharedRef<TArray<uint8>, ESPMode::ThreadSafe> FSaveByteArrayRef;
 
 
@@ -83,7 +77,6 @@ public:
 	 * Receives a payload from the connection's inbox.
 	 *
 	 * @param OutPayload Will hold the received payload, if any.
-	 *
 	 * @return true if a payload was returned, false otherwise.
 	 */
 	bool Receive( FArrayReaderPtr& OutPayload )
diff --git a/Engine/Plugins/Messaging/UdpMessaging/Source/UdpMessaging/Public/Interfaces/IUdpMessageTunnel.h b/Engine/Plugins/Messaging/UdpMessaging/Source/UdpMessaging/Public/Interfaces/IUdpMessageTunnel.h
--- a/Engine/Plugins/Messaging/UdpMessaging/Source/UdpMessaging/Public/Interfaces/IUdpMessageTunnel.h
+++ b/Engine/Plugins/Messaging/UdpMessaging/Source/UdpMessaging/Public/Interfaces/IUdpMessageTunnel.h
@@ -7,14 +7,10 @@
 #pragma once
 
 
-/**
- * Type definition for shared pointers to instances of IUdpMessageTunnel.
- */
+/** Type definition for shared pointers to instances of IUdpMessageTunnel. */
 typedef TSharedPtr<class IUdpMessageTunnel> IUdpMessageTunnelPtr;
 
-/**
- * Type definition for shared references to instances of IUdpMessageTunnel.
- */
+/** Type definition for shared references to instances of IUdpMessageTunnel. */
 typedef TSharedRef<class IUdpMessageTunnel> IUdpMessageTunnelRef;
 
 
@@ -29,9 +25,7 @@ public:
 	 * Establishes a tunnel with a remote server.
 	 *
 	 * @param RemoteEndpoint The endpoint of the server to connect to.
-	 *
 	 * @return true if the connection has been established, false otherwise.
-	 *
 	 * @see GetConnections
 	 */
 	virtual bool Connect( const FIPv4Endpoint& RemoteEndpoint ) = 0;
@@ -40,17 +34,14 @@ public:
 	 * Starts the tunnel server.
 	 *
 	 * @param LocalEndpoint The IP endpoint to listen for incoming connections on.
-	 *
-	 * @see IsServerRunning
-	 * @see StopServer
+	 * @see IsServerRunning, StopServer
 	 */
 	virtual void StartServer( const FIPv4Endpoint& LocalEndpoint ) = 0;
 
 	/**
 	 * Stops the tunnel server.
 	 *
-	 * @see IsServerRunning
-	 * @see StartServer
+	 * @see IsServerRunning, StartServer
 	 */
 	virtual void StopServer( ) = 0;
 
@@ -60,9 +51,7 @@ public:
 	 * Gets the list of all open tunnel connections.
 	 *
 	 * @param OutConnections Will hold the list of connections.
-	 *
 	 * @return The number of connections returned.
-	 *
 	 * @see Connect
 	 */
 	virtual int32 GetConnections( TArray<IUdpMessageTunnelConnectionPtr>& OutConnections ) = 0;
@@ -85,9 +74,7 @@ public:
 	 * Checks whether the tunnel server is running.
 	 *
 	 * @return true if the tunnel server is running, false otherwise.
-	 *
-	 * @see StartServer
-	 * @see StopServer
+	 * @see StartServer, StopServer
 	 */
 	virtual bool IsServerRunning( ) const = 0;
 
diff --git a/Engine/Plugins/Messaging/UdpMessaging/Source/UdpMessaging/Public/Interfaces/IUdpMessageTunnelConnection.h b/Engine/Plugins/Messaging/UdpMessaging/Source/UdpMessaging/Public/Interfaces/IUdpMessageTunnelConnection.h
--- a/Engine/Plugins/Messaging/UdpMessaging/Source/UdpMessaging/Public/Interfaces/IUdpMessageTunnelConnection.h
+++ b/Engine/Plugins/Messaging/UdpMessaging/Source/UdpMessaging/Public/Interfaces/IUdpMessageTunnelConnection.h
@@ -7,14 +7,10 @@
 #pragma once
 
 
-/**
- * Type definition for shared pointers to instances of IUdpMessageTunnelConnection.
- */
+/** Type definition for shared pointers to instances of IUdpMessageTunnelConnection. */
 typedef TSharedPtr<class IUdpMessageTunnelConnection> IUdpMessageTunnelConnectionPtr;
 
-/**
- * Type definition for shared references to instances of IUdpMessageTunnelConnection.
- */
+/** Type definition for shared references to instances of IUdpMessageTunnelConnection. */
 typedef TSharedRef<class IUdpMessageTunnelConnection> IUdpMessageTunnelConnectionRef;
 
 
