Docs: Removed file comments and added missing code documentation

Please note that file comments had no purpose in nearly all cases and just added visual clutter. The two files that had meaningful file comments had their comments moved into the corresponding classes. There are still hundreds of file comments left in other files that will be removed over time.

Also cleaned up some random stuff along the way:
- relative paths to public headers within the same module are no longer necessary (automatically discovered by UBT now)
- header guards are deprecated, use #pragma once instead (all compilers support it now)
- space between multiple template brackets is no longer required (all compilers support >> now)
- NULL to nullptr, OVERRIDE to override
- spelling errors, whitespace, line breaks

diff --git a/Engine/Plugins/Messaging/MessagingDebugger/Source/MessagingDebugger/Private/MessagingDebuggerModule.cpp b/Engine/Plugins/Messaging/MessagingDebugger/Source/MessagingDebugger/Private/MessagingDebuggerModule.cpp
--- a/Engine/Plugins/Messaging/MessagingDebugger/Source/MessagingDebugger/Private/MessagingDebuggerModule.cpp
+++ b/Engine/Plugins/Messaging/MessagingDebugger/Source/MessagingDebugger/Private/MessagingDebuggerModule.cpp
@@ -1,9 +1,5 @@
 // Copyright 1998-2014 Epic Games, Inc. All Rights Reserved.
 
-/*=============================================================================
-	MessagingDebuggerModule.cpp: Implements the FMessagingDebuggerModule class.
-=============================================================================*/
-
 #include "MessagingDebuggerPrivatePCH.h"
 #include "Runtime/Core/Public/Features/IModularFeatures.h"
 
@@ -23,9 +19,9 @@ class FMessagingDebuggerModule
 {
 public:
 
-	// Begin IModuleInterface interface
+	// IModuleInterface interface
 	
-	virtual void StartupModule( ) OVERRIDE
+	virtual void StartupModule( ) override
 	{
 		Style = MakeShareable(new FMessagingDebuggerStyle());
 
@@ -39,15 +35,12 @@ public:
 			.SetIcon(FSlateIcon(Style->GetStyleSetName(), "MessagingDebuggerTabIcon"));
 	}
 
-	virtual void ShutdownModule( ) OVERRIDE
+	virtual void ShutdownModule( ) override
 	{
 		FGlobalTabmanager::Get()->UnregisterTabSpawner(MessagingDebuggerTabName);
 		IModularFeatures::Get().UnregisterModularFeature("MessagingDebugger", this);
 		FMessagingDebuggerCommands::Unregister();
 	}
-	
-	// End IModuleInterface interface
-
 
 private:
 
diff --git a/Engine/Plugins/Messaging/MessagingDebugger/Source/MessagingDebugger/Private/MessagingDebuggerPrivatePCH.h b/Engine/Plugins/Messaging/MessagingDebugger/Source/MessagingDebugger/Private/MessagingDebuggerPrivatePCH.h
--- a/Engine/Plugins/Messaging/MessagingDebugger/Source/MessagingDebugger/Private/MessagingDebuggerPrivatePCH.h
+++ b/Engine/Plugins/Messaging/MessagingDebugger/Source/MessagingDebugger/Private/MessagingDebuggerPrivatePCH.h
@@ -1,9 +1,5 @@
 // Copyright 1998-2014 Epic Games, Inc. All Rights Reserved.
 
-/*=============================================================================
-	MessagingDebuggerPrivatePCH.h: Pre-compiled header file for the MessagingDebugger module.
-=============================================================================*/
-
 #pragma once
 
 
diff --git a/Engine/Plugins/Messaging/MessagingDebugger/Source/MessagingDebugger/Private/Models/MessagingDebuggerCommands.h b/Engine/Plugins/Messaging/MessagingDebugger/Source/MessagingDebugger/Private/Models/MessagingDebuggerCommands.h
--- a/Engine/Plugins/Messaging/MessagingDebugger/Source/MessagingDebugger/Private/Models/MessagingDebuggerCommands.h
+++ b/Engine/Plugins/Messaging/MessagingDebugger/Source/MessagingDebugger/Private/Models/MessagingDebuggerCommands.h
@@ -1,9 +1,5 @@
 // Copyright 1998-2014 Epic Games, Inc. All Rights Reserved.
 
-/*=============================================================================
-	MessagingDebuggerCommands.h: Declares the FMessagingDebuggerCommands class.
-=============================================================================*/
-
 #pragma once
 
 
@@ -24,9 +20,9 @@ public:
 
 public:
 
-	// Begin TCommands interface
+	// TCommands interface
 
-	virtual void RegisterCommands( ) OVERRIDE
+	virtual void RegisterCommands( ) override
 	{
 		UI_COMMAND(BreakDebugger, "Break", "Break the debugger at the next message", EUserInterfaceActionType::Button, FInputGesture());
 		UI_COMMAND(ClearHistory, "Clear History", "Clears the message history list", EUserInterfaceActionType::Button, FInputGesture());
@@ -36,8 +32,6 @@ public:
 		UI_COMMAND(StopDebugger, "Stop", "Stop the debugger", EUserInterfaceActionType::Button, FInputGesture());
 	}
 
-	// End TCommands interface
-
 public:
 
 	/** Holds a command that breaks at the next message. */
diff --git a/Engine/Plugins/Messaging/MessagingDebugger/Source/MessagingDebugger/Private/Models/MessagingDebuggerEndpointFilter.h b/Engine/Plugins/Messaging/MessagingDebugger/Source/MessagingDebugger/Private/Models/MessagingDebuggerEndpointFilter.h
--- a/Engine/Plugins/Messaging/MessagingDebugger/Source/MessagingDebugger/Private/Models/MessagingDebuggerEndpointFilter.h
+++ b/Engine/Plugins/Messaging/MessagingDebugger/Source/MessagingDebugger/Private/Models/MessagingDebuggerEndpointFilter.h
@@ -1,9 +1,5 @@
 // Copyright 1998-2014 Epic Games, Inc. All Rights Reserved.
 
-/*=============================================================================
-	MessagingDebuggerEndpointFilter.h: Declares the FMessagingDebuggerEndpointFilter class.
-=============================================================================*/
-
 #pragma once
 
 
diff --git a/Engine/Plugins/Messaging/MessagingDebugger/Source/MessagingDebugger/Private/Models/MessagingDebuggerMessageFilter.h b/Engine/Plugins/Messaging/MessagingDebugger/Source/MessagingDebugger/Private/Models/MessagingDebuggerMessageFilter.h
--- a/Engine/Plugins/Messaging/MessagingDebugger/Source/MessagingDebugger/Private/Models/MessagingDebuggerMessageFilter.h
+++ b/Engine/Plugins/Messaging/MessagingDebugger/Source/MessagingDebugger/Private/Models/MessagingDebuggerMessageFilter.h
@@ -1,9 +1,5 @@
 // Copyright 1998-2014 Epic Games, Inc. All Rights Reserved.
 
-/*=============================================================================
-	MessagingDebuggerMessageFilter.h: Declares the FMessagingDebuggerMessageFilter class.
-=============================================================================*/
-
 #pragma once
 
 
diff --git a/Engine/Plugins/Messaging/MessagingDebugger/Source/MessagingDebugger/Private/Models/MessagingDebuggerModel.h b/Engine/Plugins/Messaging/MessagingDebugger/Source/MessagingDebugger/Private/Models/MessagingDebuggerModel.h
--- a/Engine/Plugins/Messaging/MessagingDebugger/Source/MessagingDebugger/Private/Models/MessagingDebuggerModel.h
+++ b/Engine/Plugins/Messaging/MessagingDebugger/Source/MessagingDebugger/Private/Models/MessagingDebuggerModel.h
@@ -1,9 +1,5 @@
 // Copyright 1998-2014 Epic Games, Inc. All Rights Reserved.
 
-/*=============================================================================
-	MessagingDebuggerModel.h: Declares the FMessagingDebuggerModel class.
-=============================================================================*/
-
 #pragma once
 
 
diff --git a/Engine/Plugins/Messaging/MessagingDebugger/Source/MessagingDebugger/Private/Models/MessagingDebuggerTypeFilter.h b/Engine/Plugins/Messaging/MessagingDebugger/Source/MessagingDebugger/Private/Models/MessagingDebuggerTypeFilter.h
--- a/Engine/Plugins/Messaging/MessagingDebugger/Source/MessagingDebugger/Private/Models/MessagingDebuggerTypeFilter.h
+++ b/Engine/Plugins/Messaging/MessagingDebugger/Source/MessagingDebugger/Private/Models/MessagingDebuggerTypeFilter.h
@@ -1,9 +1,5 @@
 // Copyright 1998-2014 Epic Games, Inc. All Rights Reserved.
 
-/*=============================================================================
-	MessagingDebuggerTypeFilter.h: Declares the FMessagingDebuggerTypeFilter class.
-=============================================================================*/
-
 #pragma once
 
 
diff --git a/Engine/Plugins/Messaging/MessagingDebugger/Source/MessagingDebugger/Private/Styles/MessagingDebuggerStyle.h b/Engine/Plugins/Messaging/MessagingDebugger/Source/MessagingDebugger/Private/Styles/MessagingDebuggerStyle.h
--- a/Engine/Plugins/Messaging/MessagingDebugger/Source/MessagingDebugger/Private/Styles/MessagingDebuggerStyle.h
+++ b/Engine/Plugins/Messaging/MessagingDebugger/Source/MessagingDebugger/Private/Styles/MessagingDebuggerStyle.h
@@ -1,9 +1,5 @@
 // Copyright 1998-2014 Epic Games, Inc. All Rights Reserved.
 
-/*=============================================================================
-	MessagingDebuggerStyle.h: Declares the FMessagingDebuggerStyle class.
-=============================================================================*/
-
 #pragma once
 
 
diff --git a/Engine/Plugins/Messaging/MessagingDebugger/Source/MessagingDebugger/Private/Widgets/Breakpoints/SMessagingBreakpoints.cpp b/Engine/Plugins/Messaging/MessagingDebugger/Source/MessagingDebugger/Private/Widgets/Breakpoints/SMessagingBreakpoints.cpp
--- a/Engine/Plugins/Messaging/MessagingDebugger/Source/MessagingDebugger/Private/Widgets/Breakpoints/SMessagingBreakpoints.cpp
+++ b/Engine/Plugins/Messaging/MessagingDebugger/Source/MessagingDebugger/Private/Widgets/Breakpoints/SMessagingBreakpoints.cpp
@@ -1,9 +1,5 @@
 // Copyright 1998-2014 Epic Games, Inc. All Rights Reserved.
 
-/*=============================================================================
-	SMessagingBreakpoints.cpp: Implements the SMessagingBreakpoints class.
-=============================================================================*/
-
 #include "MessagingDebuggerPrivatePCH.h"
 
 
diff --git a/Engine/Plugins/Messaging/MessagingDebugger/Source/MessagingDebugger/Private/Widgets/Breakpoints/SMessagingBreakpoints.h b/Engine/Plugins/Messaging/MessagingDebugger/Source/MessagingDebugger/Private/Widgets/Breakpoints/SMessagingBreakpoints.h
--- a/Engine/Plugins/Messaging/MessagingDebugger/Source/MessagingDebugger/Private/Widgets/Breakpoints/SMessagingBreakpoints.h
+++ b/Engine/Plugins/Messaging/MessagingDebugger/Source/MessagingDebugger/Private/Widgets/Breakpoints/SMessagingBreakpoints.h
@@ -1,9 +1,5 @@
 // Copyright 1998-2014 Epic Games, Inc. All Rights Reserved.
 
-/*=============================================================================
-	SMessagingBreakpoints.h: Declares the SMessagingBreakpoints class.
-=============================================================================*/
-
 #pragma once
 
 
diff --git a/Engine/Plugins/Messaging/MessagingDebugger/Source/MessagingDebugger/Private/Widgets/Breakpoints/SMessagingBreakpointsTableRow.h b/Engine/Plugins/Messaging/MessagingDebugger/Source/MessagingDebugger/Private/Widgets/Breakpoints/SMessagingBreakpointsTableRow.h
--- a/Engine/Plugins/Messaging/MessagingDebugger/Source/MessagingDebugger/Private/Widgets/Breakpoints/SMessagingBreakpointsTableRow.h
+++ b/Engine/Plugins/Messaging/MessagingDebugger/Source/MessagingDebugger/Private/Widgets/Breakpoints/SMessagingBreakpointsTableRow.h
@@ -1,9 +1,5 @@
 // Copyright 1998-2014 Epic Games, Inc. All Rights Reserved.
 
-/*=============================================================================
-	SMessagingBreakpointsTableRow.h: Declares the SMessagingBreakpointsTableRow class.
-=============================================================================*/
-
 #pragma once
 
 
@@ -44,10 +40,10 @@ public:
 
 public:
 
-	// Begin SMultiColumnTableRow interface
+	// SMultiColumnTableRow interface
 
 	BEGIN_SLATE_FUNCTION_BUILD_OPTIMIZATION
-	virtual TSharedRef<SWidget> GenerateWidgetForColumn( const FName& ColumnName ) OVERRIDE
+	virtual TSharedRef<SWidget> GenerateWidgetForColumn( const FName& ColumnName ) override
 	{
 		if (ColumnName == "HitCount")
 		{
@@ -82,8 +78,6 @@ public:
 	}
 	END_SLATE_FUNCTION_BUILD_OPTIMIZATION
 
-	// End SMultiColumnTableRow interface
-
 private:
 
 	// Holds a pointer to the breakpoint that is shown in this row.
diff --git a/Engine/Plugins/Messaging/MessagingDebugger/Source/MessagingDebugger/Private/Widgets/EndpointDetails/SMessagingAddressTableRow.h b/Engine/Plugins/Messaging/MessagingDebugger/Source/MessagingDebugger/Private/Widgets/EndpointDetails/SMessagingAddressTableRow.h
--- a/Engine/Plugins/Messaging/MessagingDebugger/Source/MessagingDebugger/Private/Widgets/EndpointDetails/SMessagingAddressTableRow.h
+++ b/Engine/Plugins/Messaging/MessagingDebugger/Source/MessagingDebugger/Private/Widgets/EndpointDetails/SMessagingAddressTableRow.h
@@ -1,9 +1,5 @@
 // Copyright 1998-2014 Epic Games, Inc. All Rights Reserved.
 
-/*=============================================================================
-	SMessagingAddressTableRow.h: Declares the SMessagingAddressTableRow class.
-=============================================================================*/
-
 #pragma once
 
 
@@ -45,10 +41,10 @@ public:
 
 public:
 
-	// Begin SMultiColumnTableRow interface
+	// SMultiColumnTableRow interface
 
 	BEGIN_SLATE_FUNCTION_BUILD_OPTIMIZATION
-	virtual TSharedRef<SWidget> GenerateWidgetForColumn( const FName& ColumnName ) OVERRIDE
+	virtual TSharedRef<SWidget> GenerateWidgetForColumn( const FName& ColumnName ) override
 	{
 		if (ColumnName == "Address")
 		{
@@ -85,8 +81,6 @@ public:
 	}
 	END_SLATE_FUNCTION_BUILD_OPTIMIZATION
 
-	// End SMultiColumnTableRow interface
-
 private:
 
 	// Holds the address information.
diff --git a/Engine/Plugins/Messaging/MessagingDebugger/Source/MessagingDebugger/Private/Widgets/EndpointDetails/SMessagingEndpointDetails.cpp b/Engine/Plugins/Messaging/MessagingDebugger/Source/MessagingDebugger/Private/Widgets/EndpointDetails/SMessagingEndpointDetails.cpp
--- a/Engine/Plugins/Messaging/MessagingDebugger/Source/MessagingDebugger/Private/Widgets/EndpointDetails/SMessagingEndpointDetails.cpp
+++ b/Engine/Plugins/Messaging/MessagingDebugger/Source/MessagingDebugger/Private/Widgets/EndpointDetails/SMessagingEndpointDetails.cpp
@@ -1,9 +1,5 @@
 // Copyright 1998-2014 Epic Games, Inc. All Rights Reserved.
 
-/*=============================================================================
-	SMessagingEndpointDetails.cpp: Implements the SMessagingEndpointDetails class.
-=============================================================================*/
-
 #include "MessagingDebuggerPrivatePCH.h"
 
 
diff --git a/Engine/Plugins/Messaging/MessagingDebugger/Source/MessagingDebugger/Private/Widgets/EndpointDetails/SMessagingEndpointDetails.h b/Engine/Plugins/Messaging/MessagingDebugger/Source/MessagingDebugger/Private/Widgets/EndpointDetails/SMessagingEndpointDetails.h
--- a/Engine/Plugins/Messaging/MessagingDebugger/Source/MessagingDebugger/Private/Widgets/EndpointDetails/SMessagingEndpointDetails.h
+++ b/Engine/Plugins/Messaging/MessagingDebugger/Source/MessagingDebugger/Private/Widgets/EndpointDetails/SMessagingEndpointDetails.h
@@ -1,9 +1,5 @@
 // Copyright 1998-2014 Epic Games, Inc. All Rights Reserved.
 
-/*=============================================================================
-	SMessagingEndpointDetails.h: Declares the SMessagingEndpointDetails class.
-=============================================================================*/
-
 #pragma once
 
 
@@ -31,11 +27,9 @@ public:
 
 public:
 
-	// Begin SCompoundWidget overrides
-
-	virtual void Tick( const FGeometry& AllottedGeometry, const double InCurrentTime, const float InDeltaTime ) OVERRIDE;
+	// SCompoundWidget overrides
 
-	// End SCompoundWidget overrides
+	virtual void Tick( const FGeometry& AllottedGeometry, const double InCurrentTime, const float InDeltaTime ) override;
 
 protected:
 
diff --git a/Engine/Plugins/Messaging/MessagingDebugger/Source/MessagingDebugger/Private/Widgets/Endpoints/SMessagingEndpoints.cpp b/Engine/Plugins/Messaging/MessagingDebugger/Source/MessagingDebugger/Private/Widgets/Endpoints/SMessagingEndpoints.cpp
--- a/Engine/Plugins/Messaging/MessagingDebugger/Source/MessagingDebugger/Private/Widgets/Endpoints/SMessagingEndpoints.cpp
+++ b/Engine/Plugins/Messaging/MessagingDebugger/Source/MessagingDebugger/Private/Widgets/Endpoints/SMessagingEndpoints.cpp
@@ -1,9 +1,5 @@
 // Copyright 1998-2014 Epic Games, Inc. All Rights Reserved.
 
-/*=============================================================================
-	SMessagingEndpoints.cpp: Implements the SMessagingEndpoints class.
-=============================================================================*/
-
 #include "MessagingDebuggerPrivatePCH.h"
 
 
diff --git a/Engine/Plugins/Messaging/MessagingDebugger/Source/MessagingDebugger/Private/Widgets/Endpoints/SMessagingEndpoints.h b/Engine/Plugins/Messaging/MessagingDebugger/Source/MessagingDebugger/Private/Widgets/Endpoints/SMessagingEndpoints.h
--- a/Engine/Plugins/Messaging/MessagingDebugger/Source/MessagingDebugger/Private/Widgets/Endpoints/SMessagingEndpoints.h
+++ b/Engine/Plugins/Messaging/MessagingDebugger/Source/MessagingDebugger/Private/Widgets/Endpoints/SMessagingEndpoints.h
@@ -1,9 +1,5 @@
 // Copyright 1998-2014 Epic Games, Inc. All Rights Reserved.
 
-/*=============================================================================
-	SMessagingEndpoints.h: Declares the SMessagingEndpoints class.
-=============================================================================*/
-
 #pragma once
 
 
@@ -39,11 +35,9 @@ public:
 
 public:
 
-	// Begin SCompoundWidget overrides
-
-	virtual void Tick( const FGeometry& AllottedGeometry, const double InCurrentTime, const float InDeltaTime ) OVERRIDE;
+	// SCompoundWidget overrides
 
-	// End SCompoundWidget overrides
+	virtual void Tick( const FGeometry& AllottedGeometry, const double InCurrentTime, const float InDeltaTime ) override;
 
 protected:
 
diff --git a/Engine/Plugins/Messaging/MessagingDebugger/Source/MessagingDebugger/Private/Widgets/Endpoints/SMessagingEndpointsFilterBar.h b/Engine/Plugins/Messaging/MessagingDebugger/Source/MessagingDebugger/Private/Widgets/Endpoints/SMessagingEndpointsFilterBar.h
--- a/Engine/Plugins/Messaging/MessagingDebugger/Source/MessagingDebugger/Private/Widgets/Endpoints/SMessagingEndpointsFilterBar.h
+++ b/Engine/Plugins/Messaging/MessagingDebugger/Source/MessagingDebugger/Private/Widgets/Endpoints/SMessagingEndpointsFilterBar.h
@@ -1,9 +1,5 @@
 // Copyright 1998-2014 Epic Games, Inc. All Rights Reserved.
 
-/*=============================================================================
-	SMessagingEndpointsFilterBar.h: Declares the SMessagingEndpointsFilterBar class.
-=============================================================================*/
-
 #pragma once
 
 
diff --git a/Engine/Plugins/Messaging/MessagingDebugger/Source/MessagingDebugger/Private/Widgets/Endpoints/SMessagingEndpointsTableRow.h b/Engine/Plugins/Messaging/MessagingDebugger/Source/MessagingDebugger/Private/Widgets/Endpoints/SMessagingEndpointsTableRow.h
--- a/Engine/Plugins/Messaging/MessagingDebugger/Source/MessagingDebugger/Private/Widgets/Endpoints/SMessagingEndpointsTableRow.h
+++ b/Engine/Plugins/Messaging/MessagingDebugger/Source/MessagingDebugger/Private/Widgets/Endpoints/SMessagingEndpointsTableRow.h
@@ -1,9 +1,5 @@
 // Copyright 1998-2014 Epic Games, Inc. All Rights Reserved.
 
-/*=============================================================================
-	SMessagingEndpointsTableRow.h: Declares the SMessagingEndpointsTableRow class.
-=============================================================================*/
-
 #pragma once
 
 
@@ -47,10 +43,10 @@ public:
 
 public:
 
-	// Begin SMultiColumnTableRow interface
+	// SMultiColumnTableRow interface
 
 	BEGIN_SLATE_FUNCTION_BUILD_OPTIMIZATION
-	virtual TSharedRef<SWidget> GenerateWidgetForColumn( const FName& ColumnName ) OVERRIDE
+	virtual TSharedRef<SWidget> GenerateWidgetForColumn( const FName& ColumnName ) override
 	{
 		if (ColumnName == "Break")
 		{
@@ -114,8 +110,6 @@ public:
 	}
 	END_SLATE_FUNCTION_BUILD_OPTIMIZATION
 
-	// End SMultiColumnTableRow interface
-
 private:
 
 	// Gets the text for the Messages column.
diff --git a/Engine/Plugins/Messaging/MessagingDebugger/Source/MessagingDebugger/Private/Widgets/Graph/SMessagingGraph.cpp b/Engine/Plugins/Messaging/MessagingDebugger/Source/MessagingDebugger/Private/Widgets/Graph/SMessagingGraph.cpp
--- a/Engine/Plugins/Messaging/MessagingDebugger/Source/MessagingDebugger/Private/Widgets/Graph/SMessagingGraph.cpp
+++ b/Engine/Plugins/Messaging/MessagingDebugger/Source/MessagingDebugger/Private/Widgets/Graph/SMessagingGraph.cpp
@@ -1,9 +1,5 @@
 // Copyright 1998-2014 Epic Games, Inc. All Rights Reserved.
 
-/*=============================================================================
-	SMessagingGraph.cpp: Implements the SMessagingGraph class.
-=============================================================================*/
-
 #include "MessagingDebuggerPrivatePCH.h"
 
 
diff --git a/Engine/Plugins/Messaging/MessagingDebugger/Source/MessagingDebugger/Private/Widgets/Graph/SMessagingGraph.h b/Engine/Plugins/Messaging/MessagingDebugger/Source/MessagingDebugger/Private/Widgets/Graph/SMessagingGraph.h
--- a/Engine/Plugins/Messaging/MessagingDebugger/Source/MessagingDebugger/Private/Widgets/Graph/SMessagingGraph.h
+++ b/Engine/Plugins/Messaging/MessagingDebugger/Source/MessagingDebugger/Private/Widgets/Graph/SMessagingGraph.h
@@ -1,9 +1,5 @@
 // Copyright 1998-2014 Epic Games, Inc. All Rights Reserved.
 
-/*=============================================================================
-	SMessagingGraph.h: Declares the SMessagingGraph class.
-=============================================================================*/
-
 #pragma once
 
 
diff --git a/Engine/Plugins/Messaging/MessagingDebugger/Source/MessagingDebugger/Private/Widgets/History/SMessagingHistory.cpp b/Engine/Plugins/Messaging/MessagingDebugger/Source/MessagingDebugger/Private/Widgets/History/SMessagingHistory.cpp
--- a/Engine/Plugins/Messaging/MessagingDebugger/Source/MessagingDebugger/Private/Widgets/History/SMessagingHistory.cpp
+++ b/Engine/Plugins/Messaging/MessagingDebugger/Source/MessagingDebugger/Private/Widgets/History/SMessagingHistory.cpp
@@ -1,9 +1,5 @@
 // Copyright 1998-2014 Epic Games, Inc. All Rights Reserved.
 
-/*=============================================================================
-	SMessagingHistory.cpp: Implements the SMessagingHistory class.
-=============================================================================*/
-
 #include "MessagingDebuggerPrivatePCH.h"
 
 
diff --git a/Engine/Plugins/Messaging/MessagingDebugger/Source/MessagingDebugger/Private/Widgets/History/SMessagingHistory.h b/Engine/Plugins/Messaging/MessagingDebugger/Source/MessagingDebugger/Private/Widgets/History/SMessagingHistory.h
--- a/Engine/Plugins/Messaging/MessagingDebugger/Source/MessagingDebugger/Private/Widgets/History/SMessagingHistory.h
+++ b/Engine/Plugins/Messaging/MessagingDebugger/Source/MessagingDebugger/Private/Widgets/History/SMessagingHistory.h
@@ -1,9 +1,5 @@
 // Copyright 1998-2014 Epic Games, Inc. All Rights Reserved.
 
-/*=============================================================================
-	SMessagingHistory.h: Declares the SMessagingHistory class.
-=============================================================================*/
-
 #pragma once
 
 
diff --git a/Engine/Plugins/Messaging/MessagingDebugger/Source/MessagingDebugger/Private/Widgets/History/SMessagingHistoryFilterBar.h b/Engine/Plugins/Messaging/MessagingDebugger/Source/MessagingDebugger/Private/Widgets/History/SMessagingHistoryFilterBar.h
--- a/Engine/Plugins/Messaging/MessagingDebugger/Source/MessagingDebugger/Private/Widgets/History/SMessagingHistoryFilterBar.h
+++ b/Engine/Plugins/Messaging/MessagingDebugger/Source/MessagingDebugger/Private/Widgets/History/SMessagingHistoryFilterBar.h
@@ -1,9 +1,5 @@
 // Copyright 1998-2014 Epic Games, Inc. All Rights Reserved.
 
-/*=============================================================================
-	SMessagingHistoryFilterBar.h: Declares the SMessagingHistoryFilterBar class.
-=============================================================================*/
-
 #pragma once
 
 
diff --git a/Engine/Plugins/Messaging/MessagingDebugger/Source/MessagingDebugger/Private/Widgets/History/SMessagingHistoryTableRow.h b/Engine/Plugins/Messaging/MessagingDebugger/Source/MessagingDebugger/Private/Widgets/History/SMessagingHistoryTableRow.h
--- a/Engine/Plugins/Messaging/MessagingDebugger/Source/MessagingDebugger/Private/Widgets/History/SMessagingHistoryTableRow.h
+++ b/Engine/Plugins/Messaging/MessagingDebugger/Source/MessagingDebugger/Private/Widgets/History/SMessagingHistoryTableRow.h
@@ -1,9 +1,5 @@
 // Copyright 1998-2014 Epic Games, Inc. All Rights Reserved.
 
-/*=============================================================================
-	SMessagingHistoryTableRow.h: Declares the SMessagingHistoryTableRow class.
-=============================================================================*/
-
 #pragma once
 
 
@@ -48,9 +44,9 @@ public:
 
 public:
 
-	// Begin SWidget overrides
+	// SWidget overrides
 
-	virtual void Tick( const FGeometry& AllottedGeometry, const double InCurrentTime, const float InDeltaTime ) OVERRIDE
+	virtual void Tick( const FGeometry& AllottedGeometry, const double InCurrentTime, const float InDeltaTime ) override
 	{
 		SWidget::Tick(AllottedGeometry, InCurrentTime, InDeltaTime);
 
@@ -73,14 +69,12 @@ public:
 		}
 	}
 
-	// End SWidget overrides
-
 public:
 
-	// Begin SMultiColumnTableRow interface
+	// SMultiColumnTableRow interface
 
 	BEGIN_SLATE_FUNCTION_BUILD_OPTIMIZATION
-	virtual TSharedRef<SWidget> GenerateWidgetForColumn( const FName& ColumnName ) OVERRIDE
+	virtual TSharedRef<SWidget> GenerateWidgetForColumn( const FName& ColumnName ) override
 	{
 		if (ColumnName == "DispatchLatency")
 		{
@@ -222,8 +216,6 @@ public:
 	}
 	END_SLATE_FUNCTION_BUILD_OPTIMIZATION
 
-	// End SMultiColumnTableRow interface
-
 protected:
 
 	/** 
diff --git a/Engine/Plugins/Messaging/MessagingDebugger/Source/MessagingDebugger/Private/Widgets/Interceptors/SMessagingInterceptors.cpp b/Engine/Plugins/Messaging/MessagingDebugger/Source/MessagingDebugger/Private/Widgets/Interceptors/SMessagingInterceptors.cpp
--- a/Engine/Plugins/Messaging/MessagingDebugger/Source/MessagingDebugger/Private/Widgets/Interceptors/SMessagingInterceptors.cpp
+++ b/Engine/Plugins/Messaging/MessagingDebugger/Source/MessagingDebugger/Private/Widgets/Interceptors/SMessagingInterceptors.cpp
@@ -1,9 +1,5 @@
 // Copyright 1998-2014 Epic Games, Inc. All Rights Reserved.
 
-/*=============================================================================
-	SMessagingInterceptors.cpp: Implements the SMessagingInterceptors class.
-=============================================================================*/
-
 #include "MessagingDebuggerPrivatePCH.h"
 
 
diff --git a/Engine/Plugins/Messaging/MessagingDebugger/Source/MessagingDebugger/Private/Widgets/Interceptors/SMessagingInterceptors.h b/Engine/Plugins/Messaging/MessagingDebugger/Source/MessagingDebugger/Private/Widgets/Interceptors/SMessagingInterceptors.h
--- a/Engine/Plugins/Messaging/MessagingDebugger/Source/MessagingDebugger/Private/Widgets/Interceptors/SMessagingInterceptors.h
+++ b/Engine/Plugins/Messaging/MessagingDebugger/Source/MessagingDebugger/Private/Widgets/Interceptors/SMessagingInterceptors.h
@@ -1,9 +1,5 @@
 // Copyright 1998-2014 Epic Games, Inc. All Rights Reserved.
 
-/*=============================================================================
-	SMessagingInterceptors.h: Declares the SMessagingInterceptors class.
-=============================================================================*/
-
 #pragma once
 
 
diff --git a/Engine/Plugins/Messaging/MessagingDebugger/Source/MessagingDebugger/Private/Widgets/MessageData/SMessagingMessageData.cpp b/Engine/Plugins/Messaging/MessagingDebugger/Source/MessagingDebugger/Private/Widgets/MessageData/SMessagingMessageData.cpp
--- a/Engine/Plugins/Messaging/MessagingDebugger/Source/MessagingDebugger/Private/Widgets/MessageData/SMessagingMessageData.cpp
+++ b/Engine/Plugins/Messaging/MessagingDebugger/Source/MessagingDebugger/Private/Widgets/MessageData/SMessagingMessageData.cpp
@@ -1,9 +1,5 @@
 // Copyright 1998-2014 Epic Games, Inc. All Rights Reserved.
 
-/*=============================================================================
-	SMessagingMessageData.cpp: Implements the SMessagingMessageData class.
-=============================================================================*/
-
 #include "MessagingDebuggerPrivatePCH.h"
 
 
diff --git a/Engine/Plugins/Messaging/MessagingDebugger/Source/MessagingDebugger/Private/Widgets/MessageData/SMessagingMessageData.h b/Engine/Plugins/Messaging/MessagingDebugger/Source/MessagingDebugger/Private/Widgets/MessageData/SMessagingMessageData.h
--- a/Engine/Plugins/Messaging/MessagingDebugger/Source/MessagingDebugger/Private/Widgets/MessageData/SMessagingMessageData.h
+++ b/Engine/Plugins/Messaging/MessagingDebugger/Source/MessagingDebugger/Private/Widgets/MessageData/SMessagingMessageData.h
@@ -1,9 +1,5 @@
 // Copyright 1998-2014 Epic Games, Inc. All Rights Reserved.
 
-/*=============================================================================
-	SMessagingMessageData.h: Declares the SMessagingMessageData class.
-=============================================================================*/
-
 #pragma once
 
 
@@ -39,11 +35,9 @@ public:
 
 public:
 
-	// Begin FNotifyHook interface
-
-	virtual void NotifyPostChange( const FPropertyChangedEvent& PropertyChangedEvent, class FEditPropertyChain* PropertyThatChanged ) OVERRIDE;
+	// FNotifyHook interface
 
-	// End FNotifyHook interface
+	virtual void NotifyPostChange( const FPropertyChangedEvent& PropertyChangedEvent, class FEditPropertyChain* PropertyThatChanged ) override;
 
 private:
 
diff --git a/Engine/Plugins/Messaging/MessagingDebugger/Source/MessagingDebugger/Private/Widgets/MessageDetails/SMessagingDispatchStateTableRow.h b/Engine/Plugins/Messaging/MessagingDebugger/Source/MessagingDebugger/Private/Widgets/MessageDetails/SMessagingDispatchStateTableRow.h
--- a/Engine/Plugins/Messaging/MessagingDebugger/Source/MessagingDebugger/Private/Widgets/MessageDetails/SMessagingDispatchStateTableRow.h
+++ b/Engine/Plugins/Messaging/MessagingDebugger/Source/MessagingDebugger/Private/Widgets/MessageDetails/SMessagingDispatchStateTableRow.h
@@ -1,9 +1,5 @@
 // Copyright 1998-2014 Epic Games, Inc. All Rights Reserved.
 
-/*=============================================================================
-	SMessagingDispatchStateTableRow.h: Declares the SMessagingDispatchStateTableRow class.
-=============================================================================*/
-
 #pragma once
 
 
@@ -45,10 +41,10 @@ public:
 
 public:
 
-	// Begin SMultiColumnTableRow interface
+	// SMultiColumnTableRow interface
 
 	BEGIN_SLATE_FUNCTION_BUILD_OPTIMIZATION
-	virtual TSharedRef<SWidget> GenerateWidgetForColumn( const FName& ColumnName ) OVERRIDE
+	virtual TSharedRef<SWidget> GenerateWidgetForColumn( const FName& ColumnName ) override
 	{
 		if (ColumnName == "DispatchLatency")
 		{
@@ -97,8 +93,6 @@ public:
 	}
 	END_SLATE_FUNCTION_BUILD_OPTIMIZATION
 
-	// End SMultiColumnTableRow interface
-
 protected:
 
 	/**
@@ -132,7 +126,6 @@ protected:
 	 *
 	 * @param Seconds The time span to convert.
 	 * @return The text representation.
-	 *
 	 * @todo gmp: refactor this into FText::AsTimespan or something like that
 	 */
 	FText TimespanToReadableText( double Seconds ) const
diff --git a/Engine/Plugins/Messaging/MessagingDebugger/Source/MessagingDebugger/Private/Widgets/MessageDetails/SMessagingMessageDetails.cpp b/Engine/Plugins/Messaging/MessagingDebugger/Source/MessagingDebugger/Private/Widgets/MessageDetails/SMessagingMessageDetails.cpp
--- a/Engine/Plugins/Messaging/MessagingDebugger/Source/MessagingDebugger/Private/Widgets/MessageDetails/SMessagingMessageDetails.cpp
+++ b/Engine/Plugins/Messaging/MessagingDebugger/Source/MessagingDebugger/Private/Widgets/MessageDetails/SMessagingMessageDetails.cpp
@@ -1,9 +1,5 @@
 // Copyright 1998-2014 Epic Games, Inc. All Rights Reserved.
 
-/*=============================================================================
-	SMessagingMessageDetails.cpp: Implements the SMessagingMessageDetails class.
-=============================================================================*/
-
 #include "MessagingDebuggerPrivatePCH.h"
 
 
diff --git a/Engine/Plugins/Messaging/MessagingDebugger/Source/MessagingDebugger/Private/Widgets/MessageDetails/SMessagingMessageDetails.h b/Engine/Plugins/Messaging/MessagingDebugger/Source/MessagingDebugger/Private/Widgets/MessageDetails/SMessagingMessageDetails.h
--- a/Engine/Plugins/Messaging/MessagingDebugger/Source/MessagingDebugger/Private/Widgets/MessageDetails/SMessagingMessageDetails.h
+++ b/Engine/Plugins/Messaging/MessagingDebugger/Source/MessagingDebugger/Private/Widgets/MessageDetails/SMessagingMessageDetails.h
@@ -1,9 +1,5 @@
 // Copyright 1998-2014 Epic Games, Inc. All Rights Reserved.
 
-/*=============================================================================
-	SMessagingMessageDetails.h: Declares the SMessagingMessageDetails class.
-=============================================================================*/
-
 #pragma once
 
 
diff --git a/Engine/Plugins/Messaging/MessagingDebugger/Source/MessagingDebugger/Private/Widgets/SMessagingDebugger.cpp b/Engine/Plugins/Messaging/MessagingDebugger/Source/MessagingDebugger/Private/Widgets/SMessagingDebugger.cpp
--- a/Engine/Plugins/Messaging/MessagingDebugger/Source/MessagingDebugger/Private/Widgets/SMessagingDebugger.cpp
+++ b/Engine/Plugins/Messaging/MessagingDebugger/Source/MessagingDebugger/Private/Widgets/SMessagingDebugger.cpp
@@ -1,9 +1,5 @@
 // Copyright 1998-2014 Epic Games, Inc. All Rights Reserved.
 
-/*=============================================================================
-	SMessagingDebugger.cpp: Implements the SMessagingDebugger class.
-=============================================================================*/
-
 #include "MessagingDebuggerPrivatePCH.h"
 
 
diff --git a/Engine/Plugins/Messaging/MessagingDebugger/Source/MessagingDebugger/Private/Widgets/SMessagingDebugger.h b/Engine/Plugins/Messaging/MessagingDebugger/Source/MessagingDebugger/Private/Widgets/SMessagingDebugger.h
--- a/Engine/Plugins/Messaging/MessagingDebugger/Source/MessagingDebugger/Private/Widgets/SMessagingDebugger.h
+++ b/Engine/Plugins/Messaging/MessagingDebugger/Source/MessagingDebugger/Private/Widgets/SMessagingDebugger.h
@@ -1,9 +1,5 @@
 // Copyright 1998-2014 Epic Games, Inc. All Rights Reserved.
 
-/*=============================================================================
-	SMessagingDebugger.h: Declares the SMessagingDebugger class.
-=============================================================================*/
-
 #pragma once
 
 
@@ -40,11 +36,9 @@ public:
 
 public:
 
-	// Begin SCompoundWidget overrides
-
-	virtual void Tick( const FGeometry& AllottedGeometry, const double InCurrentTime, const float InDeltaTime ) OVERRIDE;
+	// SCompoundWidget overrides
 
-	// End SCompoundWidget overrides
+	virtual void Tick( const FGeometry& AllottedGeometry, const double InCurrentTime, const float InDeltaTime ) override;
 
 protected:
 
diff --git a/Engine/Plugins/Messaging/MessagingDebugger/Source/MessagingDebugger/Private/Widgets/Toolbar/SMessagingDebuggerToolbar.cpp b/Engine/Plugins/Messaging/MessagingDebugger/Source/MessagingDebugger/Private/Widgets/Toolbar/SMessagingDebuggerToolbar.cpp
--- a/Engine/Plugins/Messaging/MessagingDebugger/Source/MessagingDebugger/Private/Widgets/Toolbar/SMessagingDebuggerToolbar.cpp
+++ b/Engine/Plugins/Messaging/MessagingDebugger/Source/MessagingDebugger/Private/Widgets/Toolbar/SMessagingDebuggerToolbar.cpp
@@ -1,9 +1,5 @@
 // Copyright 1998-2014 Epic Games, Inc. All Rights Reserved.
 
-/*=============================================================================
-	SMessagingDebuggerToolbar.cpp: Implements the SMessagingDebuggerToolbar class.
-=============================================================================*/
-
 #include "MessagingDebuggerPrivatePCH.h"
 
 
diff --git a/Engine/Plugins/Messaging/MessagingDebugger/Source/MessagingDebugger/Private/Widgets/Toolbar/SMessagingDebuggerToolbar.h b/Engine/Plugins/Messaging/MessagingDebugger/Source/MessagingDebugger/Private/Widgets/Toolbar/SMessagingDebuggerToolbar.h
--- a/Engine/Plugins/Messaging/MessagingDebugger/Source/MessagingDebugger/Private/Widgets/Toolbar/SMessagingDebuggerToolbar.h
+++ b/Engine/Plugins/Messaging/MessagingDebugger/Source/MessagingDebugger/Private/Widgets/Toolbar/SMessagingDebuggerToolbar.h
@@ -1,9 +1,5 @@
 // Copyright 1998-2014 Epic Games, Inc. All Rights Reserved.
 
-/*=============================================================================
-	SMessagingDebuggerToolbar.h: Declares the SMessagingDebuggerToolbar class.
-=============================================================================*/
-
 #pragma once
 
 
diff --git a/Engine/Plugins/Messaging/MessagingDebugger/Source/MessagingDebugger/Private/Widgets/Types/SMessagingTypes.cpp b/Engine/Plugins/Messaging/MessagingDebugger/Source/MessagingDebugger/Private/Widgets/Types/SMessagingTypes.cpp
--- a/Engine/Plugins/Messaging/MessagingDebugger/Source/MessagingDebugger/Private/Widgets/Types/SMessagingTypes.cpp
+++ b/Engine/Plugins/Messaging/MessagingDebugger/Source/MessagingDebugger/Private/Widgets/Types/SMessagingTypes.cpp
@@ -1,9 +1,5 @@
 // Copyright 1998-2014 Epic Games, Inc. All Rights Reserved.
 
-/*=============================================================================
-	SMessagingTypes.cpp: Implements the SMessagingTypes class.
-=============================================================================*/
-
 #include "MessagingDebuggerPrivatePCH.h"
 
 
diff --git a/Engine/Plugins/Messaging/MessagingDebugger/Source/MessagingDebugger/Private/Widgets/Types/SMessagingTypes.h b/Engine/Plugins/Messaging/MessagingDebugger/Source/MessagingDebugger/Private/Widgets/Types/SMessagingTypes.h
--- a/Engine/Plugins/Messaging/MessagingDebugger/Source/MessagingDebugger/Private/Widgets/Types/SMessagingTypes.h
+++ b/Engine/Plugins/Messaging/MessagingDebugger/Source/MessagingDebugger/Private/Widgets/Types/SMessagingTypes.h
@@ -1,9 +1,5 @@
 // Copyright 1998-2014 Epic Games, Inc. All Rights Reserved.
 
-/*=============================================================================
-	SMessagingTypes.h: Declares the SMessagingTypes class.
-=============================================================================*/
-
 #pragma once
 
 
diff --git a/Engine/Plugins/Messaging/MessagingDebugger/Source/MessagingDebugger/Private/Widgets/Types/SMessagingTypesFilterBar.h b/Engine/Plugins/Messaging/MessagingDebugger/Source/MessagingDebugger/Private/Widgets/Types/SMessagingTypesFilterBar.h
--- a/Engine/Plugins/Messaging/MessagingDebugger/Source/MessagingDebugger/Private/Widgets/Types/SMessagingTypesFilterBar.h
+++ b/Engine/Plugins/Messaging/MessagingDebugger/Source/MessagingDebugger/Private/Widgets/Types/SMessagingTypesFilterBar.h
@@ -1,9 +1,5 @@
 // Copyright 1998-2014 Epic Games, Inc. All Rights Reserved.
 
-/*=============================================================================
-	SMessagingTypesFilterBar.h: Declares the SMessagingTypesFilterBar class.
-=============================================================================*/
-
 #pragma once
 
 
diff --git a/Engine/Plugins/Messaging/MessagingDebugger/Source/MessagingDebugger/Private/Widgets/Types/SMessagingTypesTableRow.h b/Engine/Plugins/Messaging/MessagingDebugger/Source/MessagingDebugger/Private/Widgets/Types/SMessagingTypesTableRow.h
--- a/Engine/Plugins/Messaging/MessagingDebugger/Source/MessagingDebugger/Private/Widgets/Types/SMessagingTypesTableRow.h
+++ b/Engine/Plugins/Messaging/MessagingDebugger/Source/MessagingDebugger/Private/Widgets/Types/SMessagingTypesTableRow.h
@@ -1,9 +1,5 @@
 // Copyright 1998-2014 Epic Games, Inc. All Rights Reserved.
 
-/*=============================================================================
-	SMessagingTypesTableRow.h: Declares the SMessagingTypesTableRow class.
-=============================================================================*/
-
 #pragma once
 
 
@@ -47,10 +43,10 @@ public:
 
 public:
 
-	// Begin SMultiColumnTableRow interface
+	// SMultiColumnTableRow interface
 
 	BEGIN_SLATE_FUNCTION_BUILD_OPTIMIZATION
-	virtual TSharedRef<SWidget> GenerateWidgetForColumn( const FName& ColumnName ) OVERRIDE
+	virtual TSharedRef<SWidget> GenerateWidgetForColumn( const FName& ColumnName ) override
 	{
 		if (ColumnName == "Break")
 		{
@@ -100,8 +96,6 @@ public:
 	}
 	END_SLATE_FUNCTION_BUILD_OPTIMIZATION
 
-	// End SMultiColumnTableRow interface
-
 private:
 
 	// Gets the text for the Messages column.
diff --git a/Engine/Plugins/Messaging/UdpMessaging/Source/UdpMessaging/Private/Shared/UdpMessagingSettings.h b/Engine/Plugins/Messaging/UdpMessaging/Source/UdpMessaging/Private/Shared/UdpMessagingSettings.h
--- a/Engine/Plugins/Messaging/UdpMessaging/Source/UdpMessaging/Private/Shared/UdpMessagingSettings.h
+++ b/Engine/Plugins/Messaging/UdpMessaging/Source/UdpMessaging/Private/Shared/UdpMessagingSettings.h
@@ -1,9 +1,5 @@
 // Copyright 1998-2014 Epic Games, Inc. All Rights Reserved.
 
-/*=============================================================================
-	UdpMessagingSettings.h: Declares the UUdpMessagingSettings class.
-=============================================================================*/
-
 #pragma once
 
 #include "UdpMessagingSettings.generated.h"
diff --git a/Engine/Plugins/Messaging/UdpMessaging/Source/UdpMessaging/Private/Transport/ReassembledUdpMessage.h b/Engine/Plugins/Messaging/UdpMessaging/Source/UdpMessaging/Private/Transport/ReassembledUdpMessage.h
--- a/Engine/Plugins/Messaging/UdpMessaging/Source/UdpMessaging/Private/Transport/ReassembledUdpMessage.h
+++ b/Engine/Plugins/Messaging/UdpMessaging/Source/UdpMessaging/Private/Transport/ReassembledUdpMessage.h
@@ -1,9 +1,5 @@
 // Copyright 1998-2014 Epic Games, Inc. All Rights Reserved.
 
-/*=============================================================================
-	ReassembledUdpMessage.h: Declares the FReassembledUdpMessage class.
-=============================================================================*/
-
 #pragma once
 
 
diff --git a/Engine/Plugins/Messaging/UdpMessaging/Source/UdpMessaging/Private/Transport/Tests/UdpMessageSegmenterTest.cpp b/Engine/Plugins/Messaging/UdpMessaging/Source/UdpMessaging/Private/Transport/Tests/UdpMessageSegmenterTest.cpp
--- a/Engine/Plugins/Messaging/UdpMessaging/Source/UdpMessaging/Private/Transport/Tests/UdpMessageSegmenterTest.cpp
+++ b/Engine/Plugins/Messaging/UdpMessaging/Source/UdpMessaging/Private/Transport/Tests/UdpMessageSegmenterTest.cpp
@@ -1,9 +1,5 @@
 // Copyright 1998-2014 Epic Games, Inc. All Rights Reserved.
 
-/*=============================================================================
-	UdpMessageSegmenterTest.cpp: Unit test for the FUdpMessageSegmenter class.
-=============================================================================*/
-
 #include "UdpMessagingPrivatePCH.h"
 #include "AutomationTest.h"
 
diff --git a/Engine/Plugins/Messaging/UdpMessaging/Source/UdpMessaging/Private/Transport/UdpMessageBeacon.cpp b/Engine/Plugins/Messaging/UdpMessaging/Source/UdpMessaging/Private/Transport/UdpMessageBeacon.cpp
--- a/Engine/Plugins/Messaging/UdpMessaging/Source/UdpMessaging/Private/Transport/UdpMessageBeacon.cpp
+++ b/Engine/Plugins/Messaging/UdpMessaging/Source/UdpMessaging/Private/Transport/UdpMessageBeacon.cpp
@@ -1,9 +1,5 @@
 // Copyright 1998-2014 Epic Games, Inc. All Rights Reserved.
 
-/*=============================================================================
-	UdpMessageHelloSender.h: Implements the FUdpMessageHelloSender class.
-=============================================================================*/
-
 #include "UdpMessagingPrivatePCH.h"
 
 
diff --git a/Engine/Plugins/Messaging/UdpMessaging/Source/UdpMessaging/Private/Transport/UdpMessageBeacon.h b/Engine/Plugins/Messaging/UdpMessaging/Source/UdpMessaging/Private/Transport/UdpMessageBeacon.h
--- a/Engine/Plugins/Messaging/UdpMessaging/Source/UdpMessaging/Private/Transport/UdpMessageBeacon.h
+++ b/Engine/Plugins/Messaging/UdpMessaging/Source/UdpMessaging/Private/Transport/UdpMessageBeacon.h
@@ -1,9 +1,5 @@
 // Copyright 1998-2014 Epic Games, Inc. All Rights Reserved.
 
-/*=============================================================================
-	UdpMessageHelloSender.h: Declares the FUdpMessageHelloSender class.
-=============================================================================*/
-
 #pragma once
 
 
@@ -52,17 +48,12 @@ public:
 
 public:
 
-	// Begin FRunnable interface
-
-	virtual bool Init( ) OVERRIDE;
-
-	virtual uint32 Run( ) OVERRIDE;
+	// FRunnable interface
 
-	virtual void Stop( ) OVERRIDE;
-
-	virtual void Exit( ) OVERRIDE { }
-	
-	// End FRunnable interface
+	virtual bool Init( ) override;
+	virtual uint32 Run( ) override;
+	virtual void Stop( ) override;
+	virtual void Exit( ) override { }
 
 protected:
 
diff --git a/Engine/Plugins/Messaging/UdpMessaging/Source/UdpMessaging/Private/Transport/UdpMessageProcessor.cpp b/Engine/Plugins/Messaging/UdpMessaging/Source/UdpMessaging/Private/Transport/UdpMessageProcessor.cpp
--- a/Engine/Plugins/Messaging/UdpMessaging/Source/UdpMessaging/Private/Transport/UdpMessageProcessor.cpp
+++ b/Engine/Plugins/Messaging/UdpMessaging/Source/UdpMessaging/Private/Transport/UdpMessageProcessor.cpp
@@ -1,9 +1,5 @@
 // Copyright 1998-2014 Epic Games, Inc. All Rights Reserved.
 
-/*=============================================================================
-	UdpMessageProcessor.cpp: Implements the FUdpMessageProcessor class.
-=============================================================================*/
-
 #include "UdpMessagingPrivatePCH.h"
 
 
diff --git a/Engine/Plugins/Messaging/UdpMessaging/Source/UdpMessaging/Private/Transport/UdpMessageProcessor.h b/Engine/Plugins/Messaging/UdpMessaging/Source/UdpMessaging/Private/Transport/UdpMessageProcessor.h
--- a/Engine/Plugins/Messaging/UdpMessaging/Source/UdpMessaging/Private/Transport/UdpMessageProcessor.h
+++ b/Engine/Plugins/Messaging/UdpMessaging/Source/UdpMessaging/Private/Transport/UdpMessageProcessor.h
@@ -1,9 +1,5 @@
 // Copyright 1998-2014 Epic Games, Inc. All Rights Reserved.
 
-/*=============================================================================
-	UdpMessageProcessor.h: Declares the FUdpMessageProcessor class.
-=============================================================================*/
-
 #pragma once
 
 
@@ -160,17 +156,12 @@ public:
 	
 public:
 
-	// Begin FRunnable interface
-
-	virtual bool Init( ) OVERRIDE;
-
-	virtual uint32 Run( ) OVERRIDE;
+	// FRunnable interface
 
-	virtual void Stop( ) OVERRIDE;
-
-	virtual void Exit( ) OVERRIDE { }
-	
-	// End FRunnable interface
+	virtual bool Init( ) override;
+	virtual uint32 Run( ) override;
+	virtual void Stop( ) override;
+	virtual void Exit( ) override { }
 
 protected:
 
diff --git a/Engine/Plugins/Messaging/UdpMessaging/Source/UdpMessaging/Private/Transport/UdpMessageResequencer.h b/Engine/Plugins/Messaging/UdpMessaging/Source/UdpMessaging/Private/Transport/UdpMessageResequencer.h
--- a/Engine/Plugins/Messaging/UdpMessaging/Source/UdpMessaging/Private/Transport/UdpMessageResequencer.h
+++ b/Engine/Plugins/Messaging/UdpMessaging/Source/UdpMessaging/Private/Transport/UdpMessageResequencer.h
@@ -1,9 +1,5 @@
 // Copyright 1998-2014 Epic Games, Inc. All Rights Reserved.
 
-/*=============================================================================
-	UdpMessageResequencer.h: Declares the FUdpMessageResequencer class.
-=============================================================================*/
-
 #pragma once
 
 
diff --git a/Engine/Plugins/Messaging/UdpMessaging/Source/UdpMessaging/Private/Transport/UdpMessageSegmenter.h b/Engine/Plugins/Messaging/UdpMessaging/Source/UdpMessaging/Private/Transport/UdpMessageSegmenter.h
--- a/Engine/Plugins/Messaging/UdpMessaging/Source/UdpMessaging/Private/Transport/UdpMessageSegmenter.h
+++ b/Engine/Plugins/Messaging/UdpMessaging/Source/UdpMessaging/Private/Transport/UdpMessageSegmenter.h
@@ -1,9 +1,5 @@
 // Copyright 1998-2014 Epic Games, Inc. All Rights Reserved.
 
-/*=============================================================================
-	UdpMessageSegmenter.h: Declares the FUdpMessageSegmenter class.
-=============================================================================*/
-
 #pragma once
 
 
diff --git a/Engine/Plugins/Messaging/UdpMessaging/Source/UdpMessaging/Private/Transport/UdpMessageTransport.cpp b/Engine/Plugins/Messaging/UdpMessaging/Source/UdpMessaging/Private/Transport/UdpMessageTransport.cpp
--- a/Engine/Plugins/Messaging/UdpMessaging/Source/UdpMessaging/Private/Transport/UdpMessageTransport.cpp
+++ b/Engine/Plugins/Messaging/UdpMessaging/Source/UdpMessaging/Private/Transport/UdpMessageTransport.cpp
@@ -1,9 +1,5 @@
 // Copyright 1998-2014 Epic Games, Inc. All Rights Reserved.
 
-/*=============================================================================
-	UdpMessageTransport.cpp: Implements the FUdpMessageTransport class.
-=============================================================================*/
-
 #include "UdpMessagingPrivatePCH.h"
 
 
diff --git a/Engine/Plugins/Messaging/UdpMessaging/Source/UdpMessaging/Private/Transport/UdpMessageTransport.h b/Engine/Plugins/Messaging/UdpMessaging/Source/UdpMessaging/Private/Transport/UdpMessageTransport.h
--- a/Engine/Plugins/Messaging/UdpMessaging/Source/UdpMessaging/Private/Transport/UdpMessageTransport.h
+++ b/Engine/Plugins/Messaging/UdpMessaging/Source/UdpMessaging/Private/Transport/UdpMessageTransport.h
@@ -1,9 +1,5 @@
 // Copyright 1998-2014 Epic Games, Inc. All Rights Reserved.
 
-/*=============================================================================
-	UdpMessageTransport.h: Declares the FUdpMessageTransport class.
-=============================================================================*/
-
 #pragma once
 
 
@@ -31,35 +27,31 @@ public:
 
 public:
 
-	// Begin ITransportMessages interface
+	// ITransportMessages interface
 
-	virtual FName GetDebugName( ) const OVERRIDE
+	virtual FName GetDebugName( ) const override
 	{
 		return "UdpMessageTransport";
 	}
 
-	virtual FOnMessageTransportMessageReceived& OnMessageReceived( ) OVERRIDE
+	virtual FOnMessageTransportMessageReceived& OnMessageReceived( ) override
 	{
 		return MessageReceivedDelegate;
 	}
 
-	virtual FOnMessageTransportNodeDiscovered& OnNodeDiscovered( ) OVERRIDE
+	virtual FOnMessageTransportNodeDiscovered& OnNodeDiscovered( ) override
 	{
 		return NodeDiscoveredDelegate;
 	}
 
-	virtual FOnMessageTransportNodeLost& OnNodeLost( ) OVERRIDE
+	virtual FOnMessageTransportNodeLost& OnNodeLost( ) override
 	{
 		return NodeLostDelegate;
 	}
 
-	virtual bool StartTransport( ) OVERRIDE;
-
-	virtual void StopTransport( ) OVERRIDE;
-
-	virtual bool TransportMessage( const IMessageDataRef& Data, const IMessageAttachmentPtr& Attachment, const TArray<FGuid>& Recipients ) OVERRIDE;
-
-	// End ITransportMessages interface
+	virtual bool StartTransport( ) override;
+	virtual void StopTransport( ) override;
+	virtual bool TransportMessage( const IMessageDataRef& Data, const IMessageAttachmentPtr& Attachment, const TArray<FGuid>& Recipients ) override;
 
 private:
 
diff --git a/Engine/Plugins/Messaging/UdpMessaging/Source/UdpMessaging/Private/Tunnel/UdpMessageTunnel.h b/Engine/Plugins/Messaging/UdpMessaging/Source/UdpMessaging/Private/Tunnel/UdpMessageTunnel.h
--- a/Engine/Plugins/Messaging/UdpMessaging/Source/UdpMessaging/Private/Tunnel/UdpMessageTunnel.h
+++ b/Engine/Plugins/Messaging/UdpMessaging/Source/UdpMessaging/Private/Tunnel/UdpMessageTunnel.h
@@ -1,9 +1,5 @@
 // Copyright 1998-2014 Epic Games, Inc. All Rights Reserved.
 
-/*=============================================================================
-	UdpMessageTunnel.h: Declares the FUdpMessageTunnel class.
-=============================================================================*/
-
 #pragma once
 
 
@@ -90,14 +86,14 @@ public:
 
 public:
 
-	// Begin FRunnable interface
+	// FRunnable interface
 
-	virtual bool Init( ) OVERRIDE
+	virtual bool Init( ) override
 	{
 		return true;
 	}
 
-	virtual uint32 Run( ) OVERRIDE
+	virtual uint32 Run( ) override
 	{
 		while (!Stopping)
 		{
@@ -117,20 +113,18 @@ public:
 		return 0;
 	}
 
-	virtual void Stop( ) OVERRIDE
+	virtual void Stop( ) override
 	{
 		Stopping = true;
 	}
 
-	virtual void Exit( ) OVERRIDE { }
-
-	// End FRunnable interface
+	virtual void Exit( ) override { }
 
 public:
 
-	// Begin IUdpMessageTunnel interface
+	// IUdpMessageTunnel interface
 
-	virtual bool Connect( const FIPv4Endpoint& RemoteEndpoint ) OVERRIDE
+	virtual bool Connect( const FIPv4Endpoint& RemoteEndpoint ) override
 	{
 		FSocket* Socket = FTcpSocketBuilder(TEXT("FUdpMessageTunnel.RemoteConnection"));
 
@@ -151,7 +145,7 @@ public:
 		return false;
 	}
 
-	virtual int32 GetConnections( TArray<IUdpMessageTunnelConnectionPtr>& OutConnections ) OVERRIDE
+	virtual int32 GetConnections( TArray<IUdpMessageTunnelConnectionPtr>& OutConnections ) override
 	{
 		FScopeLock Lock(&CriticalSection);
 
@@ -163,27 +157,27 @@ public:
 		return OutConnections.Num();
 	}
 
-	virtual uint64 GetTotalInboundBytes( ) const OVERRIDE
+	virtual uint64 GetTotalInboundBytes( ) const override
 	{
 		return TotalInboundBytes;
 	}
 
-	virtual uint64 GetTotalOutboundBytes( ) const OVERRIDE
+	virtual uint64 GetTotalOutboundBytes( ) const override
 	{
 		return TotalOutboundBytes;
 	}
 
-	virtual bool IsServerRunning( ) const OVERRIDE
+	virtual bool IsServerRunning( ) const override
 	{
 		return (Listener != nullptr);
 	}
 
-	virtual FSimpleDelegate& OnConnectionsChanged( ) OVERRIDE
+	virtual FSimpleDelegate& OnConnectionsChanged( ) override
 	{
 		return ConnectionsChangedDelegate;
 	}
 
-	virtual void StartServer( const FIPv4Endpoint& LocalEndpoint ) OVERRIDE
+	virtual void StartServer( const FIPv4Endpoint& LocalEndpoint ) override
 	{
 		StopServer();
 
@@ -191,14 +185,12 @@ public:
 		Listener->OnConnectionAccepted().BindRaw(this, &FUdpMessageTunnel::HandleListenerConnectionAccepted);
 	}
 
-	virtual void StopServer( ) OVERRIDE
+	virtual void StopServer( ) override
 	{
 		delete Listener;
 		Listener = nullptr;
 	}
 
-	// End IUdpMessageTunnel interface
-
 protected:
 
 	/**
diff --git a/Engine/Plugins/Messaging/UdpMessaging/Source/UdpMessaging/Private/Tunnel/UdpMessageTunnelConnection.h b/Engine/Plugins/Messaging/UdpMessaging/Source/UdpMessaging/Private/Tunnel/UdpMessageTunnelConnection.h
--- a/Engine/Plugins/Messaging/UdpMessaging/Source/UdpMessaging/Private/Tunnel/UdpMessageTunnelConnection.h
+++ b/Engine/Plugins/Messaging/UdpMessaging/Source/UdpMessaging/Private/Tunnel/UdpMessageTunnelConnection.h
@@ -1,9 +1,5 @@
 // Copyright 1998-2014 Epic Games, Inc. All Rights Reserved.
 
-/*=============================================================================
-	UdpMessageTunnelConnection.h: Declares the FUdpMessageTunnelConnection class.
-=============================================================================*/
-
 #pragma once
 
 
@@ -101,14 +97,14 @@ public:
 
 public:
 
-	// Begin FRunnable interface
+	// FRunnable interface
 
-	virtual bool Init( ) OVERRIDE
+	virtual bool Init( ) override
 	{
 		return (Socket != nullptr);
 	}
 
-	virtual uint32 Run( ) OVERRIDE
+	virtual uint32 Run( ) override
 	{
 		bRun = true;
 
@@ -128,31 +124,29 @@ public:
 		return 0;
 	}
 
-	virtual void Stop( ) OVERRIDE
+	virtual void Stop( ) override
 	{
 		Socket->Close();
 		bRun = false;
 	}
 
-	virtual void Exit( ) OVERRIDE { }
-
-	// End FRunnable interface
+	virtual void Exit( ) override { }
 
 public:
 
-	// Begin IUdpMessageTunnelConnection interface
+	// IUdpMessageTunnelConnection interface
 
-	virtual void Close( ) OVERRIDE
+	virtual void Close( ) override
 	{
 		Stop();
 	}
 
-	virtual uint64 GetTotalBytesReceived( ) const OVERRIDE
+	virtual uint64 GetTotalBytesReceived( ) const override
 	{
 		return TotalBytesReceived;
 	}
 
-	virtual uint64 GetTotalBytesSent( ) const OVERRIDE
+	virtual uint64 GetTotalBytesSent( ) const override
 	{
 		return TotalBytesSent;
 	}
@@ -162,7 +156,7 @@ public:
 		return RemoteEndpoint.ToText();
 	}
 
-	virtual FTimespan GetUptime( ) const OVERRIDE
+	virtual FTimespan GetUptime( ) const override
 	{
 		if (IsOpen())
 		{
@@ -172,13 +166,11 @@ public:
 		return (ClosedTime - OpenedTime);
 	}
 
-	virtual bool IsOpen( ) const OVERRIDE
+	virtual bool IsOpen( ) const override
 	{
 		return (Socket->GetConnectionState() == SCS_Connected);
 	}
 
-	// End IUdpMessageTunnelConnection interface
-
 protected:
 
 	/**
diff --git a/Engine/Plugins/Messaging/UdpMessaging/Source/UdpMessaging/Private/UdpMessagingClasses.cpp b/Engine/Plugins/Messaging/UdpMessaging/Source/UdpMessaging/Private/UdpMessagingClasses.cpp
--- a/Engine/Plugins/Messaging/UdpMessaging/Source/UdpMessaging/Private/UdpMessagingClasses.cpp
+++ b/Engine/Plugins/Messaging/UdpMessaging/Source/UdpMessaging/Private/UdpMessagingClasses.cpp
@@ -1,13 +1,8 @@
 // Copyright 1998-2014 Epic Games, Inc. All Rights Reserved.
 
-/*=============================================================================
-	UdpMessagingClasses.cpp: Implements the module's script classes.
-=============================================================================*/
-
 #include "UdpMessagingPrivatePCH.h"
 
 
-
 UUdpMessagingSettings::UUdpMessagingSettings( const class FPostConstructInitializeProperties& PCIP )
 	: Super(PCIP)
 	, EnableTransport(true)
diff --git a/Engine/Plugins/Messaging/UdpMessaging/Source/UdpMessaging/Public/Interfaces/IUdpMessageTunnel.h b/Engine/Plugins/Messaging/UdpMessaging/Source/UdpMessaging/Public/Interfaces/IUdpMessageTunnel.h
--- a/Engine/Plugins/Messaging/UdpMessaging/Source/UdpMessaging/Public/Interfaces/IUdpMessageTunnel.h
+++ b/Engine/Plugins/Messaging/UdpMessaging/Source/UdpMessaging/Public/Interfaces/IUdpMessageTunnel.h
@@ -1,9 +1,5 @@
 // Copyright 1998-2014 Epic Games, Inc. All Rights Reserved.
 
-/*=============================================================================
-	IUdpMessageTunnel.h: Declares the IUdpMessageTunnel interface.
-=============================================================================*/
-
 #pragma once
 
 
diff --git a/Engine/Plugins/Messaging/UdpMessaging/Source/UdpMessaging/Public/Interfaces/IUdpMessageTunnelConnection.h b/Engine/Plugins/Messaging/UdpMessaging/Source/UdpMessaging/Public/Interfaces/IUdpMessageTunnelConnection.h
--- a/Engine/Plugins/Messaging/UdpMessaging/Source/UdpMessaging/Public/Interfaces/IUdpMessageTunnelConnection.h
+++ b/Engine/Plugins/Messaging/UdpMessaging/Source/UdpMessaging/Public/Interfaces/IUdpMessageTunnelConnection.h
@@ -1,9 +1,5 @@
 // Copyright 1998-2014 Epic Games, Inc. All Rights Reserved.
 
-/*=============================================================================
-	IUdpMessageTunnelConnection.h: Declares the IUdpMessageTunnelConnection interface.
-=============================================================================*/
-
 #pragma once
 
 
diff --git a/Engine/Plugins/Messaging/UdpMessaging/Source/UdpMessaging/Public/UdpMessaging.h b/Engine/Plugins/Messaging/UdpMessaging/Source/UdpMessaging/Public/UdpMessaging.h
--- a/Engine/Plugins/Messaging/UdpMessaging/Source/UdpMessaging/Public/UdpMessaging.h
+++ b/Engine/Plugins/Messaging/UdpMessaging/Source/UdpMessaging/Public/UdpMessaging.h
@@ -1,9 +1,5 @@
 // Copyright 1998-2014 Epic Games, Inc. All Rights Reserved.
 
-/*=============================================================================
-	UdpMessaging.h: UdpMessaging module public header file.
-=============================================================================*/
-
 #pragma once
 
 
@@ -13,7 +9,7 @@
 #include "Networking.h"
 
 
-/* Public Interfaces
+/* Public Includes
  *****************************************************************************/
 
 #include "IUdpMessageTunnelConnection.h"
diff --git a/Engine/Plugins/Runtime/AndroidDeviceProfileSelector/Source/AndroidDeviceProfileSelector/Private/AndroidDeviceProfileSelectorPrivatePCH.h b/Engine/Plugins/Runtime/AndroidDeviceProfileSelector/Source/AndroidDeviceProfileSelector/Private/AndroidDeviceProfileSelectorPrivatePCH.h
--- a/Engine/Plugins/Runtime/AndroidDeviceProfileSelector/Source/AndroidDeviceProfileSelector/Private/AndroidDeviceProfileSelectorPrivatePCH.h
+++ b/Engine/Plugins/Runtime/AndroidDeviceProfileSelector/Source/AndroidDeviceProfileSelector/Private/AndroidDeviceProfileSelectorPrivatePCH.h
@@ -2,13 +2,12 @@
 
 #pragma once
 
-/* Public includes
- *****************************************************************************/
-#include "../Public/AndroidDeviceProfileSelector.h"
+#include "AndroidDeviceProfileSelector.h"
 
 
-/* Dependencies
+/* Private dependencies
 *****************************************************************************/
+
 #include "Core.h"
 #include "Engine.h"
 #include "ModuleManager.h"
@@ -17,4 +16,5 @@
 
 /* Private includes
  *****************************************************************************/
+
 #include "AndroidDeviceProfileSelectorModule.h"
diff --git a/Engine/Plugins/Runtime/ExampleDeviceProfileSelector/Source/ExampleDeviceProfileSelector/Private/ExampleDeviceProfileSelectorModule.cpp b/Engine/Plugins/Runtime/ExampleDeviceProfileSelector/Source/ExampleDeviceProfileSelector/Private/ExampleDeviceProfileSelectorModule.cpp
--- a/Engine/Plugins/Runtime/ExampleDeviceProfileSelector/Source/ExampleDeviceProfileSelector/Private/ExampleDeviceProfileSelectorModule.cpp
+++ b/Engine/Plugins/Runtime/ExampleDeviceProfileSelector/Source/ExampleDeviceProfileSelector/Private/ExampleDeviceProfileSelectorModule.cpp
@@ -1,11 +1,8 @@
 // Copyright 1998-2014 Epic Games, Inc. All Rights Reserved.
 
-/*=============================================================================
-	ExampleDeviceProfileSelectorModule.cpp: Implements the FExampleDeviceProfileSelectorModule class.
-=============================================================================*/
-
 #include "ExampleDeviceProfileSelectorPrivatePCH.h"
 
+
 IMPLEMENT_MODULE(FExampleDeviceProfileSelectorModule, ExampleDeviceProfileSelector);
 
 
@@ -24,4 +21,4 @@ const FString FExampleDeviceProfileSelectorModule::GetRuntimeDeviceProfileName()
 	FString ProfileName = FPlatformProperties::PlatformName();
 	UE_LOG(LogInit, Log, TEXT("Selected Device Profile: [%s]"), *ProfileName);
 	return ProfileName;
-}
\ No newline at end of file
+}
diff --git a/Engine/Plugins/Runtime/ExampleDeviceProfileSelector/Source/ExampleDeviceProfileSelector/Private/ExampleDeviceProfileSelectorModule.h b/Engine/Plugins/Runtime/ExampleDeviceProfileSelector/Source/ExampleDeviceProfileSelector/Private/ExampleDeviceProfileSelectorModule.h
--- a/Engine/Plugins/Runtime/ExampleDeviceProfileSelector/Source/ExampleDeviceProfileSelector/Private/ExampleDeviceProfileSelectorModule.h
+++ b/Engine/Plugins/Runtime/ExampleDeviceProfileSelector/Source/ExampleDeviceProfileSelector/Private/ExampleDeviceProfileSelectorModule.h
@@ -1,9 +1,5 @@
 // Copyright 1998-2014 Epic Games, Inc. All Rights Reserved.
 
-/*=============================================================================
-	ExampleDeviceProfileSelectorModule.h: Declares the FExampleDeviceProfileSelectorModule class.
-=============================================================================*/
-
 #pragma once
 
 
@@ -32,4 +28,4 @@ public:
 	virtual ~FExampleDeviceProfileSelectorModule()
 	{
 	}
-};
\ No newline at end of file
+};
diff --git a/Engine/Plugins/Runtime/ExampleDeviceProfileSelector/Source/ExampleDeviceProfileSelector/Private/ExampleDeviceProfileSelectorPrivatePCH.h b/Engine/Plugins/Runtime/ExampleDeviceProfileSelector/Source/ExampleDeviceProfileSelector/Private/ExampleDeviceProfileSelectorPrivatePCH.h
--- a/Engine/Plugins/Runtime/ExampleDeviceProfileSelector/Source/ExampleDeviceProfileSelector/Private/ExampleDeviceProfileSelectorPrivatePCH.h
+++ b/Engine/Plugins/Runtime/ExampleDeviceProfileSelector/Source/ExampleDeviceProfileSelector/Private/ExampleDeviceProfileSelectorPrivatePCH.h
@@ -1,18 +1,13 @@
 // Copyright 1998-2014 Epic Games, Inc. All Rights Reserved.
 
-/*=============================================================================
-	ExampleDeviceProfileSelectorPrivatePCH.h: Pre-compiled header file for the ExampleDeviceProfileSelector module.
-=============================================================================*/
-
 #pragma once
 
-/* Public includes
- *****************************************************************************/
-#include "../Public/ExampleDeviceProfileSelector.h"
+#include "ExampleDeviceProfileSelector.h"
 
 
-/* Dependencies
+/* Private dependencies
 *****************************************************************************/
+
 #include "Core.h"
 #include "Engine.h"
 #include "ModuleManager.h"
@@ -21,4 +16,5 @@
 
 /* Private includes
  *****************************************************************************/
-#include "ExampleDeviceProfileSelectorModule.h"
\ No newline at end of file
+
+#include "ExampleDeviceProfileSelectorModule.h"
diff --git a/Engine/Plugins/Runtime/ExampleDeviceProfileSelector/Source/ExampleDeviceProfileSelector/Public/ExampleDeviceProfileSelector.h b/Engine/Plugins/Runtime/ExampleDeviceProfileSelector/Source/ExampleDeviceProfileSelector/Public/ExampleDeviceProfileSelector.h
--- a/Engine/Plugins/Runtime/ExampleDeviceProfileSelector/Source/ExampleDeviceProfileSelector/Public/ExampleDeviceProfileSelector.h
+++ b/Engine/Plugins/Runtime/ExampleDeviceProfileSelector/Source/ExampleDeviceProfileSelector/Public/ExampleDeviceProfileSelector.h
@@ -1,13 +1,10 @@
 // Copyright 1998-2014 Epic Games, Inc. All Rights Reserved.
 
-/*=============================================================================
-	ExampleDeviceProfileSelector.h: ExampleDeviceProfileSelector module public header file.
-=============================================================================*/
-
 #pragma once
 
 
 /* Dependencies
 *****************************************************************************/
+
 #include "ModuleManager.h"
 #include "IDeviceProfileSelectorModule.h"
diff --git a/Engine/Plugins/Runtime/IOSDeviceProfileSelector/Source/IOSDeviceProfileSelector/Private/IOSDeviceProfileSelectorModule.cpp b/Engine/Plugins/Runtime/IOSDeviceProfileSelector/Source/IOSDeviceProfileSelector/Private/IOSDeviceProfileSelectorModule.cpp
--- a/Engine/Plugins/Runtime/IOSDeviceProfileSelector/Source/IOSDeviceProfileSelector/Private/IOSDeviceProfileSelectorModule.cpp
+++ b/Engine/Plugins/Runtime/IOSDeviceProfileSelector/Source/IOSDeviceProfileSelector/Private/IOSDeviceProfileSelectorModule.cpp
@@ -1,11 +1,8 @@
 // Copyright 1998-2014 Epic Games, Inc. All Rights Reserved.
 
-/*=============================================================================
-	IOSDeviceProfileSelectorModule.cpp: Implements the FIOSDeviceProfileSelectorModule class.
-=============================================================================*/
-
 #include "IOSDeviceProfileSelectorPrivatePCH.h"
 
+
 IMPLEMENT_MODULE(FIOSDeviceProfileSelectorModule, IOSDeviceProfileSelector);
 
 
@@ -30,4 +27,4 @@ FString const FIOSDeviceProfileSelectorModule::GetRuntimeDeviceProfileName()
 
 	// If no type was obtained, just select IOS as default
 	return ProfileName;
-}
\ No newline at end of file
+}
diff --git a/Engine/Plugins/Runtime/IOSDeviceProfileSelector/Source/IOSDeviceProfileSelector/Private/IOSDeviceProfileSelectorModule.h b/Engine/Plugins/Runtime/IOSDeviceProfileSelector/Source/IOSDeviceProfileSelector/Private/IOSDeviceProfileSelectorModule.h
--- a/Engine/Plugins/Runtime/IOSDeviceProfileSelector/Source/IOSDeviceProfileSelector/Private/IOSDeviceProfileSelectorModule.h
+++ b/Engine/Plugins/Runtime/IOSDeviceProfileSelector/Source/IOSDeviceProfileSelector/Private/IOSDeviceProfileSelectorModule.h
@@ -1,9 +1,5 @@
 // Copyright 1998-2014 Epic Games, Inc. All Rights Reserved.
 
-/*=============================================================================
-	IOSDeviceProfileSelectorModule.h: Declares the FIOSDeviceProfileSelectorModule class.
-=============================================================================*/
-
 #pragma once
 
 
@@ -32,4 +28,4 @@ public:
 	virtual ~FIOSDeviceProfileSelectorModule()
 	{
 	}
-};
\ No newline at end of file
+};
diff --git a/Engine/Plugins/Runtime/IOSDeviceProfileSelector/Source/IOSDeviceProfileSelector/Private/IOSDeviceProfileSelectorPrivatePCH.h b/Engine/Plugins/Runtime/IOSDeviceProfileSelector/Source/IOSDeviceProfileSelector/Private/IOSDeviceProfileSelectorPrivatePCH.h
--- a/Engine/Plugins/Runtime/IOSDeviceProfileSelector/Source/IOSDeviceProfileSelector/Private/IOSDeviceProfileSelectorPrivatePCH.h
+++ b/Engine/Plugins/Runtime/IOSDeviceProfileSelector/Source/IOSDeviceProfileSelector/Private/IOSDeviceProfileSelectorPrivatePCH.h
@@ -1,18 +1,13 @@
 // Copyright 1998-2014 Epic Games, Inc. All Rights Reserved.
 
-/*=============================================================================
-	IOSDeviceProfileSelectorPrivatePCH.h: Pre-compiled header file for the IOSDeviceProfileSelector module.
-=============================================================================*/
-
 #pragma once
 
-/* Public includes
- *****************************************************************************/
-#include "../Public/IOSDeviceProfileSelector.h"
+#include "IOSDeviceProfileSelector.h"
 
 
 /* Dependencies
 *****************************************************************************/
+
 #include "Core.h"
 #include "Engine.h"
 #include "ModuleManager.h"
@@ -21,4 +16,5 @@
 
 /* Private includes
  *****************************************************************************/
+
 #include "IOSDeviceProfileSelectorModule.h"
diff --git a/Engine/Plugins/Runtime/IOSDeviceProfileSelector/Source/IOSDeviceProfileSelector/Public/IOSDeviceProfileSelector.h b/Engine/Plugins/Runtime/IOSDeviceProfileSelector/Source/IOSDeviceProfileSelector/Public/IOSDeviceProfileSelector.h
--- a/Engine/Plugins/Runtime/IOSDeviceProfileSelector/Source/IOSDeviceProfileSelector/Public/IOSDeviceProfileSelector.h
+++ b/Engine/Plugins/Runtime/IOSDeviceProfileSelector/Source/IOSDeviceProfileSelector/Public/IOSDeviceProfileSelector.h
@@ -1,13 +1,10 @@
 // Copyright 1998-2014 Epic Games, Inc. All Rights Reserved.
 
-/*=============================================================================
-	IOSDeviceProfileSelector.h: IOSDeviceProfileSelector module public header file.
-=============================================================================*/
-
 #pragma once
 
 
 /* Dependencies
 *****************************************************************************/
+
 #include "ModuleManager.h"
 #include "IDeviceProfileSelectorModule.h"
diff --git a/Engine/Plugins/Runtime/WindowsDeviceProfileSelector/Source/WindowsDeviceProfileSelector/Private/WindowsDeviceProfileSelectorModule.cpp b/Engine/Plugins/Runtime/WindowsDeviceProfileSelector/Source/WindowsDeviceProfileSelector/Private/WindowsDeviceProfileSelectorModule.cpp
--- a/Engine/Plugins/Runtime/WindowsDeviceProfileSelector/Source/WindowsDeviceProfileSelector/Private/WindowsDeviceProfileSelectorModule.cpp
+++ b/Engine/Plugins/Runtime/WindowsDeviceProfileSelector/Source/WindowsDeviceProfileSelector/Private/WindowsDeviceProfileSelectorModule.cpp
@@ -1,9 +1,5 @@
 // Copyright 1998-2014 Epic Games, Inc. All Rights Reserved.
 
-/*=============================================================================
-	WindowsDeviceProfileSelectorModule.cpp: Implements the FWindowsDeviceProfileSelectorModule class.
-=============================================================================*/
-
 #include "WindowsDeviceProfileSelectorPrivatePCH.h"
 
 
@@ -24,4 +20,4 @@ const FString FWindowsDeviceProfileSelectorModule::GetRuntimeDeviceProfileName()
 	FString ProfileName = FPlatformProperties::PlatformName();
 	UE_LOG(LogInit, Log, TEXT("Selected Device Profile: [%s]"), *ProfileName);
 	return ProfileName;
-}
\ No newline at end of file
+}
diff --git a/Engine/Plugins/Runtime/WindowsDeviceProfileSelector/Source/WindowsDeviceProfileSelector/Private/WindowsDeviceProfileSelectorModule.h b/Engine/Plugins/Runtime/WindowsDeviceProfileSelector/Source/WindowsDeviceProfileSelector/Private/WindowsDeviceProfileSelectorModule.h
--- a/Engine/Plugins/Runtime/WindowsDeviceProfileSelector/Source/WindowsDeviceProfileSelector/Private/WindowsDeviceProfileSelectorModule.h
+++ b/Engine/Plugins/Runtime/WindowsDeviceProfileSelector/Source/WindowsDeviceProfileSelector/Private/WindowsDeviceProfileSelectorModule.h
@@ -1,9 +1,5 @@
 // Copyright 1998-2014 Epic Games, Inc. All Rights Reserved.
 
-/*=============================================================================
-	WindowsDeviceProfileSelectorModule.h: Declares the FWindowsDeviceProfileSelectorModule class.
-=============================================================================*/
-
 #pragma once
 
 
@@ -32,4 +28,4 @@ public:
 	virtual ~FWindowsDeviceProfileSelectorModule()
 	{
 	}
-};
\ No newline at end of file
+};
diff --git a/Engine/Plugins/Runtime/WindowsDeviceProfileSelector/Source/WindowsDeviceProfileSelector/Private/WindowsDeviceProfileSelectorPrivatePCH.h b/Engine/Plugins/Runtime/WindowsDeviceProfileSelector/Source/WindowsDeviceProfileSelector/Private/WindowsDeviceProfileSelectorPrivatePCH.h
--- a/Engine/Plugins/Runtime/WindowsDeviceProfileSelector/Source/WindowsDeviceProfileSelector/Private/WindowsDeviceProfileSelectorPrivatePCH.h
+++ b/Engine/Plugins/Runtime/WindowsDeviceProfileSelector/Source/WindowsDeviceProfileSelector/Private/WindowsDeviceProfileSelectorPrivatePCH.h
@@ -1,18 +1,13 @@
 // Copyright 1998-2014 Epic Games, Inc. All Rights Reserved.
 
-/*=============================================================================
-	WindowsDeviceProfileSelectorPrivatePCH.h: Pre-compiled header file for the WindowsDeviceProfileSelector module.
-=============================================================================*/
-
 #pragma once
 
-/* Public includes
- *****************************************************************************/
-#include "../Public/WindowsDeviceProfileSelector.h"
+#include "WindowsDeviceProfileSelector.h"
 
 
 /* Dependencies
 *****************************************************************************/
+
 #include "Core.h"
 #include "Engine.h"
 #include "ModuleManager.h"
@@ -21,4 +16,5 @@
 
 /* Private includes
  *****************************************************************************/
-#include "WindowsDeviceProfileSelectorModule.h"
\ No newline at end of file
+
+#include "WindowsDeviceProfileSelectorModule.h"
diff --git a/Engine/Plugins/Runtime/WindowsDeviceProfileSelector/Source/WindowsDeviceProfileSelector/Public/WindowsDeviceProfileSelector.h b/Engine/Plugins/Runtime/WindowsDeviceProfileSelector/Source/WindowsDeviceProfileSelector/Public/WindowsDeviceProfileSelector.h
--- a/Engine/Plugins/Runtime/WindowsDeviceProfileSelector/Source/WindowsDeviceProfileSelector/Public/WindowsDeviceProfileSelector.h
+++ b/Engine/Plugins/Runtime/WindowsDeviceProfileSelector/Source/WindowsDeviceProfileSelector/Public/WindowsDeviceProfileSelector.h
@@ -1,13 +1,10 @@
 // Copyright 1998-2014 Epic Games, Inc. All Rights Reserved.
 
-/*=============================================================================
-	WindowsDeviceProfileSelector.h: WindowsDeviceProfileSelector module public header file.
-=============================================================================*/
-
 #pragma once
 
 
 /* Dependencies
 *****************************************************************************/
+
 #include "ModuleManager.h"
 #include "IDeviceProfileSelectorModule.h"
diff --git a/Engine/Plugins/Slate/SlateRemote/Source/SlateRemote/Private/Server/SlateRemoteServer.cpp b/Engine/Plugins/Slate/SlateRemote/Source/SlateRemote/Private/Server/SlateRemoteServer.cpp
--- a/Engine/Plugins/Slate/SlateRemote/Source/SlateRemote/Private/Server/SlateRemoteServer.cpp
+++ b/Engine/Plugins/Slate/SlateRemote/Source/SlateRemote/Private/Server/SlateRemoteServer.cpp
@@ -1,9 +1,5 @@
 // Copyright 1998-2014 Epic Games, Inc. All Rights Reserved.
 
-/*=============================================================================
-	SlateRemoteServer.cpp: Implements the FSlateRemoteServer class.
-=============================================================================*/
-
 #include "SlateRemotePrivatePCH.h"
 
 
diff --git a/Engine/Plugins/Slate/SlateRemote/Source/SlateRemote/Private/Server/SlateRemoteServer.h b/Engine/Plugins/Slate/SlateRemote/Source/SlateRemote/Private/Server/SlateRemoteServer.h
--- a/Engine/Plugins/Slate/SlateRemote/Source/SlateRemote/Private/Server/SlateRemoteServer.h
+++ b/Engine/Plugins/Slate/SlateRemote/Source/SlateRemote/Private/Server/SlateRemoteServer.h
@@ -1,9 +1,5 @@
 // Copyright 1998-2014 Epic Games, Inc. All Rights Reserved.
 
-/*=============================================================================
-	SlateRemoteServer.h: Declares the FSlateRemoteServer class.
-=============================================================================*/
-
 #pragma once
 
 
diff --git a/Engine/Plugins/Slate/SlateRemote/Source/SlateRemote/Private/Server/SlateRemoteServerMessage.h b/Engine/Plugins/Slate/SlateRemote/Source/SlateRemote/Private/Server/SlateRemoteServerMessage.h
--- a/Engine/Plugins/Slate/SlateRemote/Source/SlateRemote/Private/Server/SlateRemoteServerMessage.h
+++ b/Engine/Plugins/Slate/SlateRemote/Source/SlateRemote/Private/Server/SlateRemoteServerMessage.h
@@ -1,9 +1,5 @@
 // Copyright 1998-2014 Epic Games, Inc. All Rights Reserved.
 
-/*=============================================================================
-	SlateRemoteServerMessage.h: Declares the FSlateRemoteServerMessage structure.
-=============================================================================*/
-
 #pragma once
 
 
diff --git a/Engine/Plugins/Slate/SlateRemote/Source/SlateRemote/Private/Shared/SlateRemoteSettings.h b/Engine/Plugins/Slate/SlateRemote/Source/SlateRemote/Private/Shared/SlateRemoteSettings.h
--- a/Engine/Plugins/Slate/SlateRemote/Source/SlateRemote/Private/Shared/SlateRemoteSettings.h
+++ b/Engine/Plugins/Slate/SlateRemote/Source/SlateRemote/Private/Shared/SlateRemoteSettings.h
@@ -1,9 +1,5 @@
 // Copyright 1998-2014 Epic Games, Inc. All Rights Reserved.
 
-/*=============================================================================
-	SlateRemoteSettings.h: Declares the USlateRemoteSettings class.
-=============================================================================*/
-
 #pragma once
 
 #include "SlateRemoteSettings.generated.h"
diff --git a/Engine/Plugins/Slate/SlateRemote/Source/SlateRemote/Private/SlateRemoteClasses.cpp b/Engine/Plugins/Slate/SlateRemote/Source/SlateRemote/Private/SlateRemoteClasses.cpp
--- a/Engine/Plugins/Slate/SlateRemote/Source/SlateRemote/Private/SlateRemoteClasses.cpp
+++ b/Engine/Plugins/Slate/SlateRemote/Source/SlateRemote/Private/SlateRemoteClasses.cpp
@@ -1,9 +1,5 @@
 // Copyright 1998-2014 Epic Games, Inc. All Rights Reserved.
 
-/*=============================================================================
-	SlateRemoteClasses.cpp: Implements the module's script classes.
-=============================================================================*/
-
 #include "SlateRemotePrivatePCH.h"
 
 
diff --git a/Engine/Plugins/Slate/SlateRemote/Source/SlateRemote/Private/SlateRemoteModule.cpp b/Engine/Plugins/Slate/SlateRemote/Source/SlateRemote/Private/SlateRemoteModule.cpp
--- a/Engine/Plugins/Slate/SlateRemote/Source/SlateRemote/Private/SlateRemoteModule.cpp
+++ b/Engine/Plugins/Slate/SlateRemote/Source/SlateRemote/Private/SlateRemoteModule.cpp
@@ -1,9 +1,5 @@
 // Copyright 1998-2014 Epic Games, Inc. All Rights Reserved.
 
-/*=============================================================================
-	SlateRemoteModule.cpp: Implements the SlateRemote module.
-=============================================================================*/
-
 #include "SlateRemotePrivatePCH.h"
 
 
@@ -18,7 +14,7 @@ class FSlateRemoteModule
 {
 public:
 
-	// Begin IModuleInterface interface
+	// IModuleInterface interface
 
 	virtual void StartupModule( ) override
 	{
@@ -73,8 +69,6 @@ public:
 		return true;
 	}
 
-	// End IModuleInterface interface
-
 protected:
 
 	/**
diff --git a/Engine/Plugins/Slate/SlateRemote/Source/SlateRemote/Private/SlateRemotePrivatePCH.h b/Engine/Plugins/Slate/SlateRemote/Source/SlateRemote/Private/SlateRemotePrivatePCH.h
--- a/Engine/Plugins/Slate/SlateRemote/Source/SlateRemote/Private/SlateRemotePrivatePCH.h
+++ b/Engine/Plugins/Slate/SlateRemote/Source/SlateRemote/Private/SlateRemotePrivatePCH.h
@@ -1,13 +1,9 @@
 // Copyright 1998-2014 Epic Games, Inc. All Rights Reserved.
 
-/*=============================================================================
-	SlateRemotePrivatePCH.h: Pre-compiled header file for the SlateRemote module.
-=============================================================================*/
-
 #pragma once
 
 
-/* Private Dependencies
+/* Private dependencies
  *****************************************************************************/
 
 #include "Core.h"
@@ -20,7 +16,7 @@
 #include "SocketSubsystem.h"
 
 
-/* Private Constants
+/* Private constants
  *****************************************************************************/
 
 /**
diff --git a/Engine/Source/Developer/AutomationWindow/Private/AutomationFilter.h b/Engine/Source/Developer/AutomationWindow/Private/AutomationFilter.h
--- a/Engine/Source/Developer/AutomationWindow/Private/AutomationFilter.h
+++ b/Engine/Source/Developer/AutomationWindow/Private/AutomationFilter.h
@@ -10,48 +10,16 @@ class FAutomationFilter : public IFilter< const TSharedPtr< class IAutomationRep
 {
 public:
 
+	/**
+	 * Default constructor.
+	 */
 	FAutomationFilter()
 		: OnlySmokeTests( false )
 		, ShowErrors( false )
 		, ShowWarnings( false )
-	{}
-
-	// Begin IFilter functions
-	DECLARE_DERIVED_EVENT(FAutomationFilter, IFilter< const TSharedPtr< class IAutomationReport >& >::FChangedEvent, FChangedEvent);
-	virtual FChangedEvent& OnChanged() OVERRIDE { return ChangedEvent; }
-
-	/**
-	 * Checks if the report passes the filter
-	 * @param InReport - The automation report
-	 * @return true if it passes the test
-	 */
-	virtual bool PassesFilter( const TSharedPtr< IAutomationReport >& InReport ) const OVERRIDE
-	{
-		bool FilterPassed = true;
-		if( OnlySmokeTests && !InReport->IsSmokeTest( ) )
-		{
-			FilterPassed = false;
-		}
-
-		if (ShowWarnings && ShowErrors)
-		{
-			FilterPassed = InReport->HasWarnings() || InReport->HasErrors();
-		}
-		else if (!ShowWarnings && ShowErrors)
-		{
-			FilterPassed = InReport->HasErrors();
-		}
-		else if (ShowWarnings && !ShowErrors)
-		{
-			// Do not show report as a warning if it should be highlighted as an Error!
-			FilterPassed = InReport->HasWarnings() && !InReport->HasErrors();
-		}
-
-		return FilterPassed;
-	}
-
-	// End IFilter functions
+	{ }
 
+public:
 
 	/**
 	 * Set if we should only show warnings
@@ -107,6 +75,38 @@ public:
 		return OnlySmokeTests;
 	}
 
+public:
+
+	// IFilter interface
+
+	DECLARE_DERIVED_EVENT(FAutomationFilter, IFilter< const TSharedPtr< class IAutomationReport >& >::FChangedEvent, FChangedEvent);
+	virtual FChangedEvent& OnChanged() override { return ChangedEvent; }
+
+	virtual bool PassesFilter( const TSharedPtr< IAutomationReport >& InReport ) const override
+	{
+		bool FilterPassed = true;
+		if( OnlySmokeTests && !InReport->IsSmokeTest( ) )
+		{
+			FilterPassed = false;
+		}
+
+		if (ShowWarnings && ShowErrors)
+		{
+			FilterPassed = InReport->HasWarnings() || InReport->HasErrors();
+		}
+		else if (!ShowWarnings && ShowErrors)
+		{
+			FilterPassed = InReport->HasErrors();
+		}
+		else if (ShowWarnings && !ShowErrors)
+		{
+			// Do not show report as a warning if it should be highlighted as an Error!
+			FilterPassed = InReport->HasWarnings() && !InReport->HasErrors();
+		}
+
+		return FilterPassed;
+	}
+
 private:
 
 	/**	The event that broadcasts whenever a change occurs to the filter */
diff --git a/Engine/Source/Developer/AutomationWindow/Private/AutomationWindowModule.cpp b/Engine/Source/Developer/AutomationWindow/Private/AutomationWindowModule.cpp
--- a/Engine/Source/Developer/AutomationWindow/Private/AutomationWindowModule.cpp
+++ b/Engine/Source/Developer/AutomationWindow/Private/AutomationWindowModule.cpp
@@ -15,7 +15,7 @@ class FAutomationWindowModule
 {
 public:
 
-	// Begin IAutomationWindowModule interface
+	// IAutomationWindowModule interface
 
 	virtual TSharedRef<class SWidget> CreateAutomationWindow( const IAutomationControllerManagerRef& AutomationController, const ISessionManagerRef& SessionManager ) OVERRIDE
 	{
@@ -34,11 +34,9 @@ public:
 
 	virtual void SetAutomationWindowTab(TWeakPtr<class SDockTab> AutomationWindowTab) OVERRIDE { AutomationWindowTabPtr = AutomationWindowTab; }
 
-	// End IAutomationWindowModule interface
-
 public:
 
-	// Begin IModuleInterface interface
+	// IModuleInterface interface
 
 	virtual void StartupModule( ) OVERRIDE
 	{
@@ -49,8 +47,6 @@ public:
 		ShutdownDelegate.ExecuteIfBound();	
 	}
 
-	// End IModuleInterface interface
-
 private:
 
 	// Holds the DockTab for the AutomationWindow
diff --git a/Engine/Source/Developer/AutomationWindow/Private/AutomationWindowPrivatePCH.h b/Engine/Source/Developer/AutomationWindow/Private/AutomationWindowPrivatePCH.h
--- a/Engine/Source/Developer/AutomationWindow/Private/AutomationWindowPrivatePCH.h
+++ b/Engine/Source/Developer/AutomationWindow/Private/AutomationWindowPrivatePCH.h
@@ -1,13 +1,9 @@
 // Copyright 1998-2014 Epic Games, Inc. All Rights Reserved.
 
-/*=============================================================================
-	AutomationWindowPrivatePCH.h: Pre-compiled header file for the AutomationWindow module.
-=============================================================================*/
+#include "AutomationWindow.h"
 
-#include "../Public/AutomationWindow.h"
 
-
-/* Dependencies
+/* Private dependencies
  *****************************************************************************/
 
 #include "Messaging.h"
diff --git a/Engine/Source/Developer/CrashTracker/Private/AVIHandler.cpp b/Engine/Source/Developer/CrashTracker/Private/AVIHandler.cpp
--- a/Engine/Source/Developer/CrashTracker/Private/AVIHandler.cpp
+++ b/Engine/Source/Developer/CrashTracker/Private/AVIHandler.cpp
@@ -1,9 +1,5 @@
 // Copyright 1998-2014 Epic Games, Inc. All Rights Reserved.
 
-/*=============================================================================
-	AVIHandler.cpp: Implements the AVI handling functions.
-=============================================================================*/
-
 #include "CrashTrackerPrivatePCH.h"
 
 
diff --git a/Engine/Source/Developer/CrashTracker/Private/AVIHandler.h b/Engine/Source/Developer/CrashTracker/Private/AVIHandler.h
--- a/Engine/Source/Developer/CrashTracker/Private/AVIHandler.h
+++ b/Engine/Source/Developer/CrashTracker/Private/AVIHandler.h
@@ -1,15 +1,12 @@
 // Copyright 1998-2014 Epic Games, Inc. All Rights Reserved.
 
-/*=============================================================================
-	AVIHandler.h: Declares the AVI handling functions.
-=============================================================================*/
-
 #pragma once
 
 
 /**
  * Given a series of JPEG compressed images, this will wrap them with an AVI header
- * and write them out as a Motion JPEG avi video.
+ * and write them out as a Motion JPEG AVI video.
+ *
  * @param CompressedFrames	An ordered array of JPEG images
  * @param OutputPath		The path with which to dump out the video
  * @param Width				The width of the images
diff --git a/Engine/Source/Developer/CrashTracker/Private/CrashTrackerModule.cpp b/Engine/Source/Developer/CrashTracker/Private/CrashTrackerModule.cpp
--- a/Engine/Source/Developer/CrashTracker/Private/CrashTrackerModule.cpp
+++ b/Engine/Source/Developer/CrashTracker/Private/CrashTrackerModule.cpp
@@ -1,12 +1,9 @@
 // Copyright 1998-2014 Epic Games, Inc. All Rights Reserved.
 
-/*=============================================================================
-	CrashTrackerModule.cpp: Implements the FCrashTrackerModule class.
-=============================================================================*/
-
 #include "CrashTrackerPrivatePCH.h"
 #include "EngineBuildSettings.h"
 
+
 #if (PLATFORM_WINDOWS || PLATFORM_MAC) && !UE_BUILD_MINIMAL
 	#define CRASH_TRACKER_SUPPORTED 1
 #else
@@ -20,7 +17,7 @@ class FCrashTrackerEventLogger : public FStabilityEventLogger
 public:
 	FCrashTrackerEventLogger(TSharedPtr<FCrashVideoCapture> InVideoCapture = NULL);
 
-	virtual void Log(EEventLog::Type Event, const FString& AdditionalContent, TSharedPtr<SWidget> Widget) OVERRIDE;
+	virtual void Log(EEventLog::Type Event, const FString& AdditionalContent, TSharedPtr<SWidget> Widget) override;
 	
 	void OnHandleError();
 	void OnHandleEnsure();
@@ -77,29 +74,34 @@ class FCrashTrackerModule
 	: public ICrashTrackerModule
 {
 public:
-	// Begin IModuleInterface interface
-	virtual void StartupModule() OVERRIDE;
-	virtual void ShutdownModule() OVERRIDE;
-	// End IModuleInterface interface
-
-
-	// Begin ICrashTracker interface
-	virtual void Update(float DeltaSeconds) OVERRIDE;
-	virtual void ForceCompleteCapture() OVERRIDE;
-	virtual void SetCrashTrackingEnabled(bool bEnabled) OVERRIDE;
-	virtual bool IsVideoCaptureAvailable() const OVERRIDE;
-	virtual bool IsCurrentlyCapturing() const OVERRIDE;
-	virtual void InvalidateCrashTrackerFrame() OVERRIDE;
-	virtual EWriteUserCaptureVideoError::Type WriteUserVideoNow( FString& OutFinalSaveName ) OVERRIDE;
-	// End ICrashTracker interface
+
+	// IModuleInterface interface
+
+	virtual void StartupModule() override;
+	virtual void ShutdownModule() override;
+
+public:
+
+	// ICrashTracker interface
+
+	virtual void Update(float DeltaSeconds) override;
+	virtual void ForceCompleteCapture() override;
+	virtual void SetCrashTrackingEnabled(bool bEnabled) override;
+	virtual bool IsVideoCaptureAvailable() const override;
+	virtual bool IsCurrentlyCapturing() const override;
+	virtual void InvalidateCrashTrackerFrame() override;
+	virtual EWriteUserCaptureVideoError::Type WriteUserVideoNow( FString& OutFinalSaveName ) override;
 
 private:
+
 	// The video capture instance
 	TSharedPtr<FCrashVideoCapture> VideoCapture;
+
 	// The event logger instance
 	TSharedPtr<FCrashTrackerEventLogger> EventLogger;
 };
 
+
 void FCrashTrackerModule::StartupModule()
 {
 	FCrashVideoCapture::CaptureVideoPath = FPaths::GameLogDir() + TEXT("CrashVideo.avi");
diff --git a/Engine/Source/Developer/CrashTracker/Private/CrashTrackerPrivatePCH.h b/Engine/Source/Developer/CrashTracker/Private/CrashTrackerPrivatePCH.h
--- a/Engine/Source/Developer/CrashTracker/Private/CrashTrackerPrivatePCH.h
+++ b/Engine/Source/Developer/CrashTracker/Private/CrashTrackerPrivatePCH.h
@@ -1,16 +1,11 @@
 // Copyright 1998-2014 Epic Games, Inc. All Rights Reserved.
 
-/*=============================================================================
-	CrashTrackerPrivatePCH.h: Pre-compiled header file for the CrashTracker module.
-=============================================================================*/
-
 #pragma once
 
-
-#include "../Public/CrashTracker.h"
+#include "CrashTracker.h"
 
 
-/* Dependencies
+/* Private dependencies
  *****************************************************************************/
 
 #include "ImageWrapper.h"
diff --git a/Engine/Source/Developer/CrashTracker/Private/CrashVideoCapture.cpp b/Engine/Source/Developer/CrashTracker/Private/CrashVideoCapture.cpp
--- a/Engine/Source/Developer/CrashTracker/Private/CrashVideoCapture.cpp
+++ b/Engine/Source/Developer/CrashTracker/Private/CrashVideoCapture.cpp
@@ -1,9 +1,5 @@
 // Copyright 1998-2014 Epic Games, Inc. All Rights Reserved.
 
-/*=============================================================================
-	CrashVideoCapture.cpp: Implements the FCrashVideoCapture class.
-=============================================================================*/
-
 #include "CrashTrackerPrivatePCH.h"
 
 
diff --git a/Engine/Source/Developer/CrashTracker/Private/CrashVideoCapture.h b/Engine/Source/Developer/CrashTracker/Private/CrashVideoCapture.h
--- a/Engine/Source/Developer/CrashTracker/Private/CrashVideoCapture.h
+++ b/Engine/Source/Developer/CrashTracker/Private/CrashVideoCapture.h
@@ -1,9 +1,5 @@
 // Copyright 1998-2014 Epic Games, Inc. All Rights Reserved.
 
-/*=============================================================================
-	CrashVideoCapture.h: Declares the FCrashVideoCapture class.
-=============================================================================*/
-
 #pragma once
 
 
@@ -11,7 +7,7 @@ DECLARE_LOG_CATEGORY_EXTERN(LogCrashTracker, Log, All);
 
 
 /**
- * A buffer for holding a compressed video frame
+ * A buffer for holding a compressed video frame.
  */
 class FCompressedDataFrame
 {
@@ -29,7 +25,9 @@ public:
 
 
 
-/** An asynchronous task for compressing an image */
+/**
+ * An asynchronous task for compressing an image.
+ */
 class FAsyncImageCompress : public FNonAbandonableTask
 {
 public:
@@ -49,8 +47,7 @@ private:
 
 
 /**
- * The crash tracker handles slate key logging
- * and if CRASH_TRACKER_ENABLED, also crash video capture
+ * The crash tracker handles slate key logging and if CRASH_TRACKER_ENABLED, also crash video capture.
  */
 class FCrashVideoCapture
 {
diff --git a/Engine/Source/Developer/CrashTracker/Public/CrashTracker.h b/Engine/Source/Developer/CrashTracker/Public/CrashTracker.h
--- a/Engine/Source/Developer/CrashTracker/Public/CrashTracker.h
+++ b/Engine/Source/Developer/CrashTracker/Public/CrashTracker.h
@@ -1,9 +1,5 @@
 // Copyright 1998-2014 Epic Games, Inc. All Rights Reserved.
 
-/*=============================================================================
-	CrashTracker.h: CrashTracker module public header file.
-=============================================================================*/
-
 #pragma once
 
 
diff --git a/Engine/Source/Developer/CrashTracker/Public/Interfaces/ICrashTrackerModule.h b/Engine/Source/Developer/CrashTracker/Public/Interfaces/ICrashTrackerModule.h
--- a/Engine/Source/Developer/CrashTracker/Public/Interfaces/ICrashTrackerModule.h
+++ b/Engine/Source/Developer/CrashTracker/Public/Interfaces/ICrashTrackerModule.h
@@ -1,9 +1,5 @@
 // Copyright 1998-2014 Epic Games, Inc. All Rights Reserved.
 
-/*=============================================================================
-	ICrashTrackerModule.h: Declares the ICrashTrackerModule interface.
-=============================================================================*/
-
 #pragma once
 
 
@@ -16,19 +12,13 @@ namespace EWriteUserCaptureVideoError
 	{
 		None,
 
-		/**
-		 * The video capture instance was invalid.
-		 */
+		/** The video capture instance was invalid. */
 		VideoCaptureInvalid,
 
-		/**
-		 * Video capture is not running.
-		 */
+		/** Video capture is not running. */
 		CaptureNotRunning,
 
-		/**
-		 * Failed to create destination directory for the video.
-		 */
+		/** Failed to create destination directory for the video. */
 		FailedToCreateDirectory,
 	};
 }
@@ -68,20 +58,19 @@ public:
 	virtual void SetCrashTrackingEnabled(bool bEnabled) = 0;
 
 	/**
-	 * Updates the crash tracker, which may trigger the capture of a frame
-	 * Will also begin capturing if it hasn't begun already
+	 * Updates the crash tracker, which may trigger the capture of a frame.
+	 * Will also begin capturing if it hasn't begun already.
 	 */
 	virtual void Update(float DeltaSeconds) = 0;
 
 	/**
-	 * Write the current crash tracker as a user video
+	 * Write the current crash tracker as a user video.
 	 *
-	 * @param OutFinalSaveName	The path and name of the file the video was saved as 
-	 * @return If the video was successfully written will return EWriteUserCaptureVideoError::None, otherwise returns an error code
+	 * @param OutFinalSaveName	The path and name of the file the video was saved as.
+	 * @return If the video was successfully written will return EWriteUserCaptureVideoError::None, otherwise returns an error code.
 	 */
 	virtual EWriteUserCaptureVideoError::Type WriteUserVideoNow( FString& OutFinalSaveName ) = 0;
 
-
 public:
 
 	/**
diff --git a/Engine/Source/Developer/DeviceManager/Private/DeviceManagerModule.cpp b/Engine/Source/Developer/DeviceManager/Private/DeviceManagerModule.cpp
--- a/Engine/Source/Developer/DeviceManager/Private/DeviceManagerModule.cpp
+++ b/Engine/Source/Developer/DeviceManager/Private/DeviceManagerModule.cpp
@@ -1,9 +1,5 @@
 // Copyright 1998-2014 Epic Games, Inc. All Rights Reserved.
 
-/*=============================================================================
-	DeviceManagerModule.cpp: Implements the FDeviceManagerModule class.
-=============================================================================*/
-
 #include "DeviceManagerPrivatePCH.h"
 
 
@@ -18,9 +14,9 @@ class FDeviceManagerModule
 {
 public:
 
-	// Begin IModuleInterface interface
+	// IModuleInterface interface
 
-	virtual void StartupModule( ) OVERRIDE
+	virtual void StartupModule( ) override
 	{
 		// @todo gmp: implement an IoC container
 		ITargetDeviceServicesModule& TargetDeviceServicesModule = FModuleManager::LoadModuleChecked<ITargetDeviceServicesModule>(TEXT("TargetDeviceServices"));
@@ -33,31 +29,26 @@ public:
 			.SetIcon(FSlateIcon(FEditorStyle::GetStyleSetName(), "DeviceDetails.TabIcon"));
 	}
 
-	virtual void ShutdownModule( ) OVERRIDE
+	virtual void ShutdownModule( ) override
 	{
 		FGlobalTabmanager::Get()->UnregisterTabSpawner(DeviceManagerTabName);
 	}
 
-	// End IModuleInterface interface
-
 public:
 
-	// Begin IDeviceManagerModule interface
+	// IDeviceManagerModule interface
 
-	virtual TSharedRef<SWidget> CreateDeviceManager( const ITargetDeviceServiceManagerRef& DeviceServiceManager, const TSharedRef<SDockTab>& ConstructUnderMajorTab, const TSharedPtr<SWindow>& ConstructUnderWindow ) OVERRIDE
+	virtual TSharedRef<SWidget> CreateDeviceManager( const ITargetDeviceServiceManagerRef& DeviceServiceManager, const TSharedRef<SDockTab>& ConstructUnderMajorTab, const TSharedPtr<SWindow>& ConstructUnderWindow ) override
 	{
 		return SNew(SDeviceManager, DeviceServiceManager, ConstructUnderMajorTab, ConstructUnderWindow);
 	}
 
-	// End IDeviceManagerModule interface
-
 private:
 
 	/**
 	 * Creates a new device manager tab.
 	 *
-	 * @param SpawnTabArgs - The arguments for the tab to spawn.
-	 *
+	 * @param SpawnTabArgs The arguments for the tab to spawn.
 	 * @return The spawned tab.
 	 */
 	TSharedRef<SDockTab> SpawnDeviceManagerTab( const FSpawnTabArgs& SpawnTabArgs )
diff --git a/Engine/Source/Developer/DeviceManager/Private/DeviceManagerPrivatePCH.h b/Engine/Source/Developer/DeviceManager/Private/DeviceManagerPrivatePCH.h
--- a/Engine/Source/Developer/DeviceManager/Private/DeviceManagerPrivatePCH.h
+++ b/Engine/Source/Developer/DeviceManager/Private/DeviceManagerPrivatePCH.h
@@ -1,12 +1,8 @@
 // Copyright 1998-2014 Epic Games, Inc. All Rights Reserved.
 
-/*=============================================================================
-	DeviceManagerPrivatePCH.h: Pre-compiled header file for the DeviceManager module.
-=============================================================================*/
-
 #pragma once
 
-#include "../Public/DeviceManager.h"
+#include "DeviceManager.h"
 
 
 /* Dependencies
diff --git a/Engine/Source/Developer/DeviceManager/Private/Models/DeviceBrowserFilter.h b/Engine/Source/Developer/DeviceManager/Private/Models/DeviceBrowserFilter.h
--- a/Engine/Source/Developer/DeviceManager/Private/Models/DeviceBrowserFilter.h
+++ b/Engine/Source/Developer/DeviceManager/Private/Models/DeviceBrowserFilter.h
@@ -1,9 +1,5 @@
 // Copyright 1998-2014 Epic Games, Inc. All Rights Reserved.
 
-/*=============================================================================
-	DeviceBrowserFilter.h: Declares the FDeviceBrowserFilter class.
-=============================================================================*/
-
 #pragma once
 
 
