UdpMessaging: Removed file header comments

diff --git a/Engine/Plugins/Messaging/UdpMessaging/Source/UdpMessaging/Private/UdpMessagingModule.cpp b/Engine/Plugins/Messaging/UdpMessaging/Source/UdpMessaging/Private/UdpMessagingModule.cpp
--- a/Engine/Plugins/Messaging/UdpMessaging/Source/UdpMessaging/Private/UdpMessagingModule.cpp
+++ b/Engine/Plugins/Messaging/UdpMessaging/Source/UdpMessaging/Private/UdpMessagingModule.cpp
@@ -1,9 +1,5 @@
 // Copyright 1998-2014 Epic Games, Inc. All Rights Reserved.
 
-/*=============================================================================
-	UdpMessagingModule.cpp: Implements the FUdpMessagingModule class.
-=============================================================================*/
-
 #include "UdpMessagingPrivatePCH.h"
 
 
@@ -18,7 +14,7 @@ class FUdpMessagingModule
 {
 public:
 
-	// Begin IModuleInterface interface
+	// IModuleInterface interface
 
 	virtual void StartupModule( ) override
 	{
@@ -82,13 +78,9 @@ public:
 		return true;
 	}
 
-	// End IModuleInterface interface
-
 protected:
 
-	/**
-	 * Initializes the message bridge with the current settings.
-	 */
+	/** Initializes the message bridge with the current settings. */
 	void InitializeBridge( )
 	{
 		ShutdownBridge();
@@ -141,9 +133,7 @@ protected:
 			.UsingTransport(MakeShareable(new FUdpMessageTransport(UnicastEndpoint, MulticastEndpoint, Settings->MulticastTimeToLive)));
 	}
 
-	/**
-	 * Initializes the message tunnel with the current settings.
-	 */
+	/** Initializes the message tunnel with the current settings. */
 	void InitializeTunnel( )
 	{
 		ShutdownTunnel();
@@ -195,9 +185,7 @@ protected:
 		}
 	}
 
-	/**
-	 * Restarts the bridge and tunnel services.
-	 */
+	/** Restarts the bridge and tunnel services. */
 	void RestartServices( )
 	{
 		const UUdpMessagingSettings& Settings = *GetDefault<UUdpMessagingSettings>();
@@ -231,7 +219,6 @@ protected:
 	 * Checks whether networked message transport is supported.
 	 *
 	 * @todo gmp: this should be moved into an Engine module, so it can be shared with other transports
-	 *
 	 * @return true if networked transport is supported, false otherwise.
 	 */
 	bool SupportsNetworkedTransport( ) const
@@ -263,9 +250,7 @@ protected:
 		return true;
 	}
 
-	/**
-	 * Shuts down the message bridge.
-	 */
+	/** Shuts down the message bridge. */
 	void ShutdownBridge( )
 	{
 		if (MessageBridge.IsValid())
@@ -276,9 +261,7 @@ protected:
 		}
 	}
 
-	/**
-	 * Shuts down the message tunnel.
-	 */
+	/** Shuts down the message tunnel. */
 	void ShutdownTunnel( )
 	{
 		if (MessageTunnel.IsValid())
diff --git a/Engine/Plugins/Messaging/UdpMessaging/Source/UdpMessaging/Private/UdpMessagingPrivatePCH.h b/Engine/Plugins/Messaging/UdpMessaging/Source/UdpMessaging/Private/UdpMessagingPrivatePCH.h
--- a/Engine/Plugins/Messaging/UdpMessaging/Source/UdpMessaging/Private/UdpMessagingPrivatePCH.h
+++ b/Engine/Plugins/Messaging/UdpMessaging/Source/UdpMessaging/Private/UdpMessagingPrivatePCH.h
@@ -1,9 +1,5 @@
 // Copyright 1998-2014 Epic Games, Inc. All Rights Reserved.
 
-/*=============================================================================
-	UdpMessagingPrivatePCH.h: Pre-compiled header file for the UdpMessaging module.
-=============================================================================*/
-
 #pragma once
 
 #include "UdpMessaging.h"
@@ -24,14 +20,10 @@
 /* Private Constants
  *****************************************************************************/
 
-/**
- * Defines the default IP endpoint for multicast traffic.
- */
+/** Defines the default IP endpoint for multicast traffic. */
 #define UDP_MESSAGING_DEFAULT_MULTICAST_ENDPOINT FIPv4Endpoint(FIPv4Address(230, 0, 0, 1), 6666)
 
-/**
- * Defines the protocol version of the UDP message transport.
- */
+/** Defines the protocol version of the UDP message transport. */
 #define UDP_MESSAGING_TRANSPORT_PROTOCOL_VERSION 9
 
 
