UdpMessaging: Documentation cleanup pass

diff --git a/Engine/Plugins/Messaging/UdpMessaging/Source/UdpMessaging/Private/Transport/UdpMessageBeacon.h b/Engine/Plugins/Messaging/UdpMessaging/Source/UdpMessaging/Private/Transport/UdpMessageBeacon.h
--- a/Engine/Plugins/Messaging/UdpMessaging/Source/UdpMessaging/Private/Transport/UdpMessageBeacon.h
+++ b/Engine/Plugins/Messaging/UdpMessaging/Source/UdpMessaging/Private/Transport/UdpMessageBeacon.h
@@ -22,10 +22,8 @@ public:
 	 */
 	FUdpMessageBeacon( FSocket* InSocket, const FGuid& InSocketId, const FIPv4Endpoint& InMulticastEndpoint );
 
-	/**
-	 * Destructor.
-	 */
-	virtual ~FUdpMessageBeacon( );
+	/** Destructor. */
+	virtual ~FUdpMessageBeacon();
 
 public:
 
@@ -34,7 +32,7 @@ public:
 	 *
 	 * @return Beacon interval.
 	 */
-	FTimespan GetBeaconInterval( )
+	FTimespan GetBeaconInterval()
 	{
 		return BeaconInterval;
 	}
@@ -50,10 +48,10 @@ public:
 
 	// FRunnable interface
 
-	virtual bool Init( ) override;
-	virtual uint32 Run( ) override;
-	virtual void Stop( ) override;
-	virtual void Exit( ) override { }
+	virtual bool Init() override;
+	virtual uint32 Run() override;
+	virtual void Stop() override;
+	virtual void Exit() override { }
 
 protected:
 
@@ -66,41 +64,41 @@ protected:
 
 private:
 
-	// Holds the calculated interval between Hello segments.
+	/** Holds the calculated interval between Hello segments. */
 	FTimespan BeaconInterval;
 
-	// Holds an event signaling that an endpoint left.
+	/** Holds an event signaling that an endpoint left. */
 	FEvent* EndpointLeftEvent;
 
-	// Holds the number of known endpoints when NextHelloTime was last calculated.
+	/** Holds the number of known endpoints when NextHelloTime was last calculated. */
 	int32 LastEndpointCount;
 
-	// Holds the time at which the last Hello segment was sent.
+	/** Holds the time at which the last Hello segment was sent. */
 	FDateTime LastHelloSent;
 
-	// Holds the multicast address and port number to send to.
+	/** Holds the multicast address and port number to send to. */
 	TSharedPtr<FInternetAddr> MulticastAddress;
 
-	// Holds the time at which the next Hello segment must be sent.
+	/** Holds the time at which the next Hello segment must be sent. */
 	FDateTime NextHelloTime;
 
-	// Holds local node identifier.
+	/** Holds local node identifier. */
 	FGuid NodeId;
 
-	// Holds the socket used to send Hello segments.
+	/** Holds the socket used to send Hello segments. */
 	FSocket* Socket;
 
-	// Holds a flag indicating that the thread is stopping.
+	/** Holds a flag indicating that the thread is stopping. */
 	bool Stopping;
 
-	// Holds the thread object.
+	/** Holds the thread object. */
 	FRunnableThread* Thread;
 
 private:
 	
-	// Defines the time interval per endpoint.
+	/** Defines the time interval per endpoint. */
 	static const FTimespan IntervalPerEndpoint;
 
-	// Defines the minimum interval for Hello segments.
+	/** Defines the minimum interval for Hello segments. */
 	static const FTimespan MinimumInterval;
 };
diff --git a/Engine/Plugins/Messaging/UdpMessaging/Source/UdpMessaging/Private/Transport/UdpMessageProcessor.h b/Engine/Plugins/Messaging/UdpMessaging/Source/UdpMessaging/Private/Transport/UdpMessageProcessor.h
--- a/Engine/Plugins/Messaging/UdpMessaging/Source/UdpMessaging/Private/Transport/UdpMessageProcessor.h
+++ b/Engine/Plugins/Messaging/UdpMessaging/Source/UdpMessaging/Private/Transport/UdpMessageProcessor.h
@@ -31,7 +31,7 @@ class FUdpMessageProcessor
 		TMap<int32, TSharedPtr<FUdpMessageSegmenter> > Segmenters;
 
 		// Default constructor.
-		FNodeInfo( )
+		FNodeInfo()
 			: LastSegmentReceivedTime(FDateTime::MinValue())
 			, NodeId()
 		{ }
@@ -61,7 +61,7 @@ class FUdpMessageProcessor
 
 
 		// Default constructor.
-		FInboundSegment( ) { }
+		FInboundSegment() { }
 
 		// Creates and initializes a new instance.
 		FInboundSegment( const FArrayReaderPtr& InData, const FIPv4Endpoint& InSender )
@@ -82,7 +82,7 @@ class FUdpMessageProcessor
 
 
 		// Default constructor.
-		FOutboundMessage( ) { }
+		FOutboundMessage() { }
 
 		// Creates and initializes a new instance.
 		FOutboundMessage( const IMessageDataRef& InMessageData, const FGuid& InRecipientId )
@@ -102,10 +102,8 @@ public:
 	 */
 	FUdpMessageProcessor( FSocket* InSocket, const FGuid& InNodeId, const FIPv4Endpoint& InMulticastEndpoint );
 
-	/**
-	 * Destructor.
-	 */
-	~FUdpMessageProcessor( );
+	/** Destructor. */
+	~FUdpMessageProcessor();
 
 public:
 
@@ -129,7 +127,7 @@ public:
 
 public:
 
-	FOnMessageTransportNodeDiscovered& OnNodeDiscovered( )
+	FOnMessageTransportNodeDiscovered& OnNodeDiscovered()
 	{
 		return NodeDiscoveredDelegate;
 	}
@@ -139,7 +137,7 @@ public:
 	 *
 	 * @param Delegate The delegate to set.
 	 */
-	FOnMessageTransportNodeLost& OnNodeLost( )
+	FOnMessageTransportNodeLost& OnNodeLost()
 	{
 		return NodeLostDelegate;
 	}
@@ -149,7 +147,7 @@ public:
 	 *
 	 * @param Delegate The delegate to set.
 	 */
-	FOnMessageTransportMessageReceived& OnMessageReceived( )
+	FOnMessageTransportMessageReceived& OnMessageReceived()
 	{
 		return MessageReceivedDelegate;
 	}
@@ -158,10 +156,10 @@ public:
 
 	// FRunnable interface
 
-	virtual bool Init( ) override;
-	virtual uint32 Run( ) override;
-	virtual void Stop( ) override;
-	virtual void Exit( ) override { }
+	virtual bool Init() override;
+	virtual uint32 Run() override;
+	virtual void Stop() override;
+	virtual void Exit() override { }
 
 protected:
 
@@ -170,7 +168,6 @@ protected:
 	 *
 	 * @param MessageId The identifier of the message to acknowledge.
 	 * @param NodeInfo Details for the node to send the acknowledgment to.
-	 *
 	 * @todo gmp: batch multiple of these into a single message
 	 */
 	void AcknowledgeReceipt( int32 MessageId, const FNodeInfo& NodeInfo );
@@ -180,17 +177,13 @@ protected:
 	 *
 	 * @return Wait time.
 	 */
-	FTimespan CalculateWaitTime( ) const;
+	FTimespan CalculateWaitTime() const;
 
-	/**
-	 * Consumes all inbound segments.
-	 */
-	void ConsumeInboundSegments( );
+	/** Consumes all inbound segments. */
+	void ConsumeInboundSegments();
 
-	/**
-	 * Consumes all outbound messages.
-	 */
-	void ConsumeOutboundMessages( );
+	/** Consumes all outbound messages. */
+	void ConsumeOutboundMessages();
 
 	/**
 	 * Filters the specified message segment.
@@ -274,10 +267,8 @@ protected:
 	 */
 	void RemoveKnownNode( const FGuid& NodeId );
 
-	/**
-	 * Updates all known remote nodes.
-	 */
-	void UpdateKnownNodes( );
+	/** Updates all known remote nodes. */
+	void UpdateKnownNodes();
 
 	/**
 	 * Updates all segmenters of the specified node.
@@ -286,75 +277,73 @@ protected:
 	 */
 	void UpdateSegmenters( FNodeInfo& NodeInfo );
 
-	/**
-	 * Updates all static remote nodes.
-	 */
-	void UpdateStaticNodes( );
+	/** Updates all static remote nodes. */
+	void UpdateStaticNodes();
 
 private:
 
-	// Handles message data state changes.
-	void HandleMessageDataStateChanged( );
+	/** Handles message data state changes. */
+	void HandleMessageDataStateChanged();
 
 private:
 
-	// Holds the queue of outbound messages.
+	// Holds the queue of outbound messages. */
 	TQueue<FInboundSegment, EQueueMode::Mpsc> InboundSegments;
 
-	// Holds the queue of outbound messages.
+	// Holds the queue of outbound messages. */
 	TQueue<FOutboundMessage, EQueueMode::Mpsc> OutboundMessages;
 
 private:
 
-	// Holds the hello sender.
+	/** Holds the hello sender. */
 	FUdpMessageBeacon* Beacon;
 
-	// Holds the current time.
+	/** Holds the current time. */
 	FDateTime CurrentTime;
 
-	// Holds the collection of known remote nodes.
+	/** Holds the collection of known remote nodes. */
 	TMap<FGuid, FNodeInfo> KnownNodes;
 
-	// Holds the last sent message number.
+	/** Holds the last sent message number. */
 	int32 LastSentMessage;
 
-	// Holds the local node identifier.
+	/** Holds the local node identifier. */
 	FGuid LocalNodeId;
 
-	// Holds the multicast endpoint.
+	/** Holds the multicast endpoint. */
 	FIPv4Endpoint MulticastEndpoint;
 
-	// Holds the socket sender.
+	/** Holds the socket sender. */
 	FUdpSocketSender* Sender;
 
-	// Holds the network socket used to transport messages.
+	/** Holds the network socket used to transport messages. */
 	FSocket* Socket;
 
-	// Holds the collection of static remote nodes.
+	/** Holds the collection of static remote nodes. */
 	TMap<FIPv4Endpoint, FNodeInfo> StaticNodes;
 
-	// Holds a flag indicating that the thread is stopping.
+	/** Holds a flag indicating that the thread is stopping. */
 	bool Stopping;
 
-	// Holds the thread object.
+	/** Holds the thread object. */
 	FRunnableThread* Thread;
 
-	// Holds an event signaling that inbound messages need to be processed.
+	/** Holds an event signaling that inbound messages need to be processed. */
 	FEvent* WorkEvent;
 
 private:
 
-	// Holds a delegate to be invoked when a message was received on the transport channel.
+	/** Holds a delegate to be invoked when a message was received on the transport channel. */
 	FOnMessageTransportMessageReceived MessageReceivedDelegate;
 
-	// Holds a delegate to be invoked when a network node was discovered.
+	/** Holds a delegate to be invoked when a network node was discovered. */
 	FOnMessageTransportNodeDiscovered NodeDiscoveredDelegate;
 
-	// Holds a delegate to be invoked when a network node was lost.
+	/** Holds a delegate to be invoked when a network node was lost. */
 	FOnMessageTransportNodeLost NodeLostDelegate;
 
 private:
 
-	// Defines the maximum number of Hello segments that can be dropped before a remote endpoint is considered dead.
+	/** Defines the maximum number of Hello segments that can be dropped before a remote endpoint is considered dead. */
 	static const int32 DeadHelloIntervals;
 };
diff --git a/Engine/Plugins/Messaging/UdpMessaging/Source/UdpMessaging/Private/Transport/UdpMessageResequencer.h b/Engine/Plugins/Messaging/UdpMessaging/Source/UdpMessaging/Private/Transport/UdpMessageResequencer.h
--- a/Engine/Plugins/Messaging/UdpMessaging/Source/UdpMessaging/Private/Transport/UdpMessageResequencer.h
+++ b/Engine/Plugins/Messaging/UdpMessaging/Source/UdpMessaging/Private/Transport/UdpMessageResequencer.h
@@ -10,9 +10,7 @@ class FUdpMessageResequencer
 {
 public:
 
-	/**
-	 * Default constructor.
-	 */
+	/** Default constructor. */
 	FUdpMessageResequencer() { }
 
 	/**
@@ -67,9 +65,7 @@ public:
 		return (Message->GetSequence() == NextSequence);
 	}
 
-	/**
-	 * Resets the re-sequencer.
-	 */
+	/** Resets the re-sequencer. */
 	void Reset()
 	{
 		MessageHeap.Reset();
@@ -79,7 +75,7 @@ public:
 
 private:
 
-	// Helper for ordering messages by their sequence numbers.
+	/** Helper for ordering messages by their sequence numbers. */
 	struct FSequenceComparer
 	{
 		bool operator () (const FReassembledUdpMessagePtr& A, const FReassembledUdpMessagePtr& B) const
@@ -90,15 +86,15 @@ private:
 
 private:
 
-	// Holds the next expected sequence number.
+	/** Holds the next expected sequence number. */
 	uint64 NextSequence;
 
-	// Holds the highest received sequence number.
+	/** Holds the highest received sequence number. */
 	uint64 HighestReceivedSequence;
 
-	// Holds the messages that need to be resequenced.
+	/** Holds the messages that need to be resequenced. */
 	TArray<FReassembledUdpMessagePtr> MessageHeap;
 
-	// Holds the maximum resequence window size.
+	/** Holds the maximum resequence window size. */
 	uint16 WindowSize;
 };
diff --git a/Engine/Plugins/Messaging/UdpMessaging/Source/UdpMessaging/Private/Transport/UdpMessageSegmenter.h b/Engine/Plugins/Messaging/UdpMessaging/Source/UdpMessaging/Private/Transport/UdpMessageSegmenter.h
--- a/Engine/Plugins/Messaging/UdpMessaging/Source/UdpMessaging/Private/Transport/UdpMessageSegmenter.h
+++ b/Engine/Plugins/Messaging/UdpMessaging/Source/UdpMessaging/Private/Transport/UdpMessageSegmenter.h
@@ -13,10 +13,8 @@ class FUdpMessageSegmenter
 {
 public:
 
-	/**
-	 * Default constructor.
-	 */
-	FUdpMessageSegmenter( )
+	/** Default constructor. */
+	FUdpMessageSegmenter()
 		: MessageReader(nullptr)
 	{ }
 
@@ -31,10 +29,8 @@ public:
 		, SegmentSize(InSegmentSize)
 	{ }
 
-	/**
-	 * Destructor.
-	 */
-	~FUdpMessageSegmenter( )
+	/** Destructor. */
+	~FUdpMessageSegmenter()
 	{
 		if (MessageReader != nullptr)
 		{
@@ -119,9 +115,7 @@ public:
 		return PendingSegments.Num();
 	}
 
-	/**
-	 * Initializes the segmenter.
-	 */
+	/** Initializes the segmenter. */
 	void Initialize()
 	{
 		if (MessageReader != nullptr)
@@ -209,18 +203,18 @@ public:
 
 private:
 
-	// Holds the message.
+	/** Holds the message. */
 	IMessageDataPtr Message;
 
-	// temp hack to support new transport API
+	/** temp hack to support new transport API. */
 	FArchive* MessageReader;
 
-	// Holds an array of bits that indicate which segments still need to be sent.
+	/** Holds an array of bits that indicate which segments still need to be sent. */
 	TBitArray<> PendingSegments;
 
-	// Holds the number of segments that haven't been sent yet.
+	/** Holds the number of segments that haven't been sent yet. */
 	uint16 PendingSegmentsCount;
 
-	// Holds the segment size.
+	/** Holds the segment size. */
 	uint16 SegmentSize;
 };
diff --git a/Engine/Plugins/Messaging/UdpMessaging/Source/UdpMessaging/Private/Transport/UdpMessageTransport.h b/Engine/Plugins/Messaging/UdpMessaging/Source/UdpMessaging/Private/Transport/UdpMessageTransport.h
--- a/Engine/Plugins/Messaging/UdpMessaging/Source/UdpMessaging/Private/Transport/UdpMessageTransport.h
+++ b/Engine/Plugins/Messaging/UdpMessaging/Source/UdpMessaging/Private/Transport/UdpMessageTransport.h
@@ -20,102 +20,100 @@ public:
 	 */
 	FUdpMessageTransport( const FIPv4Endpoint& InLocalEndpoint, const FIPv4Endpoint& InMulticastEndpoint, uint8 InMulticastTtl );
 
-	/**
-	 * Destructor.
-	 */
-	virtual ~FUdpMessageTransport( );
+	/** Destructor. */
+	virtual ~FUdpMessageTransport();
 
 public:
 
 	// ITransportMessages interface
 
-	virtual FName GetDebugName( ) const override
+	virtual FName GetDebugName() const override
 	{
 		return "UdpMessageTransport";
 	}
 
-	virtual FOnMessageTransportMessageReceived& OnMessageReceived( ) override
+	virtual FOnMessageTransportMessageReceived& OnMessageReceived() override
 	{
 		return MessageReceivedDelegate;
 	}
 
-	virtual FOnMessageTransportNodeDiscovered& OnNodeDiscovered( ) override
+	virtual FOnMessageTransportNodeDiscovered& OnNodeDiscovered() override
 	{
 		return NodeDiscoveredDelegate;
 	}
 
-	virtual FOnMessageTransportNodeLost& OnNodeLost( ) override
+	virtual FOnMessageTransportNodeLost& OnNodeLost() override
 	{
 		return NodeLostDelegate;
 	}
 
-	virtual bool StartTransport( ) override;
-	virtual void StopTransport( ) override;
+	virtual bool StartTransport() override;
+	virtual void StopTransport() override;
 	virtual bool TransportMessage( const IMessageDataRef& Data, const IMessageAttachmentPtr& Attachment, const TArray<FGuid>& Recipients ) override;
 
 private:
 
-	// Handles received transport messages.
+	// Handles received transport messages. */
 	void HandleProcessorMessageReceived( FArchive& MessageData, const IMessageAttachmentPtr& Attachment, const FGuid& NodeId )
 	{
 		MessageReceivedDelegate.ExecuteIfBound(MessageData, Attachment, NodeId);
 	}
 
-	// Handles discovered transport endpoints.
+	// Handles discovered transport endpoints. */
 	void HandleProcessorNodeDiscovered( const FGuid& DiscoveredNodeId )
 	{
 		NodeDiscoveredDelegate.ExecuteIfBound(DiscoveredNodeId);
 	}
 
-	// Handles lost transport endpoints.
+	// Handles lost transport endpoints. */
 	void HandleProcessorNodeLost( const FGuid& LostNodeId )
 	{
 		NodeLostDelegate.ExecuteIfBound(LostNodeId);
 	}
 
-	// Handles received socket data.
+	// Handles received socket data. */
 	void HandleSocketDataReceived(const FArrayReaderPtr& Data, const FIPv4Endpoint& Sender);
 
 private:
 
-	// Holds the local endpoint to receive messages on.
+	/** Holds the local endpoint to receive messages on. */
 	FIPv4Endpoint LocalEndpoint;
 
-	// Holds the message processor.
+	/** Holds the message processor. */
 	FUdpMessageProcessor* MessageProcessor;
 
-	// Holds the message processor thread.
+	/** Holds the message processor thread. */
 	FRunnableThread* MessageProcessorThread;
 
-	// Holds the multicast endpoint.
+	/** Holds the multicast endpoint. */
 	FIPv4Endpoint MulticastEndpoint;
 
-	// Holds the multicast socket receiver.
+	/** Holds the multicast socket receiver. */
 	FUdpSocketReceiver* MulticastReceiver;
 
-	// Holds the multicast socket.
+	/** Holds the multicast socket. */
 	FSocket* MulticastSocket;
 
-	// Holds the multicast time to live.
+	/** Holds the multicast time to live. */
 	uint8 MulticastTtl;
 
-	// Holds a pointer to the socket sub-system.
+	/** Holds a pointer to the socket sub-system. */
 	ISocketSubsystem* SocketSubsystem;
 
-	// Holds the unicast socket receiver.
+	/** Holds the unicast socket receiver. */
 	FUdpSocketReceiver* UnicastReceiver;
 
-	// Holds the unicast socket.
+	/** Holds the unicast socket. */
 	FSocket* UnicastSocket;
 
 private:
 
-	// Holds a delegate to be invoked when a message was received on the transport channel.
+	/** Holds a delegate to be invoked when a message was received on the transport channel. */
 	FOnMessageTransportMessageReceived MessageReceivedDelegate;
 
-	// Holds a delegate to be invoked when a network node was discovered.
+	/** Holds a delegate to be invoked when a network node was discovered. */
 	FOnMessageTransportNodeDiscovered NodeDiscoveredDelegate;
 
-	// Holds a delegate to be invoked when a network node was lost.
+	/** Holds a delegate to be invoked when a network node was lost. */
 	FOnMessageTransportNodeLost NodeLostDelegate;
 };
diff --git a/Engine/Plugins/Messaging/UdpMessaging/Source/UdpMessaging/Private/Tunnel/UdpMessageTunnel.h b/Engine/Plugins/Messaging/UdpMessaging/Source/UdpMessaging/Private/Tunnel/UdpMessageTunnel.h
--- a/Engine/Plugins/Messaging/UdpMessaging/Source/UdpMessaging/Private/Tunnel/UdpMessageTunnel.h
+++ b/Engine/Plugins/Messaging/UdpMessaging/Source/UdpMessaging/Private/Tunnel/UdpMessageTunnel.h
@@ -59,10 +59,8 @@ public:
 		Thread = FRunnableThread::Create(this, TEXT("FUdpMessageTunnel"), 128 * 1024, TPri_AboveNormal);
 	}
 
-	/**
-	 * Destructor.
-	 */
-	~FUdpMessageTunnel( )
+	/** Destructor. */
+	~FUdpMessageTunnel()
 	{
 		if (Thread != nullptr)
 		{
@@ -88,12 +86,12 @@ public:
 
 	// FRunnable interface
 
-	virtual bool Init( ) override
+	virtual bool Init() override
 	{
 		return true;
 	}
 
-	virtual uint32 Run( ) override
+	virtual uint32 Run() override
 	{
 		while (!Stopping)
 		{
@@ -113,12 +111,12 @@ public:
 		return 0;
 	}
 
-	virtual void Stop( ) override
+	virtual void Stop() override
 	{
 		Stopping = true;
 	}
 
-	virtual void Exit( ) override { }
+	virtual void Exit() override { }
 
 public:
 
@@ -157,22 +155,22 @@ public:
 		return OutConnections.Num();
 	}
 
-	virtual uint64 GetTotalInboundBytes( ) const override
+	virtual uint64 GetTotalInboundBytes() const override
 	{
 		return TotalInboundBytes;
 	}
 
-	virtual uint64 GetTotalOutboundBytes( ) const override
+	virtual uint64 GetTotalOutboundBytes() const override
 	{
 		return TotalOutboundBytes;
 	}
 
-	virtual bool IsServerRunning( ) const override
+	virtual bool IsServerRunning() const override
 	{
 		return (Listener != nullptr);
 	}
 
-	virtual FSimpleDelegate& OnConnectionsChanged( ) override
+	virtual FSimpleDelegate& OnConnectionsChanged() override
 	{
 		return ConnectionsChangedDelegate;
 	}
@@ -185,7 +183,7 @@ public:
 		Listener->OnConnectionAccepted().BindRaw(this, &FUdpMessageTunnel::HandleListenerConnectionAccepted);
 	}
 
-	virtual void StopServer( ) override
+	virtual void StopServer() override
 	{
 		delete Listener;
 		Listener = nullptr;
@@ -209,10 +207,8 @@ protected:
 		}
 	}
 
-	/**
-	 * Receives all pending payloads from the tunnels and forwards them to the local message bus.
-	 */
-	void TcpToUdp( )
+	/** Receives all pending payloads from the tunnels and forwards them to the local message bus. */
+	void TcpToUdp()
 	{
 		TArray<FUdpMessageTunnelConnectionPtr>::TIterator It(Connections);
 		FArrayReaderPtr Payload;
@@ -333,10 +329,8 @@ protected:
 		}
 	}
 
-	/**
-	 * Updates all active and pending connections.
-	 */
-	void UpdateConnections( )
+	/** Updates all active and pending connections. */
+	void UpdateConnections()
 	{
 		bool ConnectionsChanged = false;
 
@@ -373,7 +367,7 @@ protected:
 
 private:
 
-	// Callback for accepted connections to the local tunnel server.
+	/** Callback for accepted connections to the local tunnel server. */
 	bool HandleListenerConnectionAccepted( FSocket* ClientSocket, const FIPv4Endpoint& ClientEndpoint )
 	{
 		PendingConnections.Enqueue(MakeShareable(new FUdpMessageTunnelConnection(ClientSocket, ClientEndpoint)));
@@ -383,50 +377,50 @@ private:
 
 private:
 
-	// Holds the list of open tunnel connections.
+	/** Holds the list of open tunnel connections. */
 	TArray<FUdpMessageTunnelConnectionPtr> Connections;
 
-	// Holds a critical section for serializing access to the connections list.
+	/** Holds a critical section for serializing access to the connections list. */
 	FCriticalSection CriticalSection;
 
-	// Holds the current time.
+	/** Holds the current time. */
 	FDateTime CurrentTime;
 
-	// Holds the local listener for incoming tunnel connections.
+	/** Holds the local listener for incoming tunnel connections. */
 	FTcpListener* Listener;
 
-	// Holds a collection of information for local transport nodes.
+	/** Holds a collection of information for local transport nodes. */
 	TMap<FGuid, FNodeInfo> LocalNodes;
 
-	// Holds the multicast endpoint.
+	/** Holds the multicast endpoint. */
 	FIPv4Endpoint MulticastEndpoint;
 
-	// Holds the multicast socket.
+	/** Holds the multicast socket. */
 	FSocket* MulticastSocket;
 
-	// Holds a queue of pending connections.
+	/** Holds a queue of pending connections. */
 	TQueue<FUdpMessageTunnelConnectionPtr, EQueueMode::Mpsc> PendingConnections;
 
-	// Holds a collection of information for remote transport nodes.
+	/** Holds a collection of information for remote transport nodes. */
 	TMap<FGuid, FNodeInfo> RemoteNodes;
 
-	// Holds a flag indicating that the thread is stopping.
+	/** Holds a flag indicating that the thread is stopping. */
 	bool Stopping;
 
-	// Holds the thread object.
+	/** Holds the thread object. */
 	FRunnableThread* Thread;
 
-	// Holds the total number of bytes that were received from tunnels.
+	/** Holds the total number of bytes that were received from tunnels. */
 	uint64 TotalInboundBytes;
 
-	// Holds the total number of bytes that were sent out through tunnels.
+	/** Holds the total number of bytes that were sent out through tunnels. */
 	uint64 TotalOutboundBytes;
 
-	// Holds the unicast socket.
+	/** Holds the unicast socket. */
 	FSocket* UnicastSocket;
 
 private:
 
-	// Holds a delegate that is executed when the list of incoming connections changed.
+	/** Holds a delegate that is executed when the list of incoming connections changed. */
 	FSimpleDelegate ConnectionsChangedDelegate;
 };
diff --git a/Engine/Plugins/Messaging/UdpMessaging/Source/UdpMessaging/Private/Tunnel/UdpMessageTunnelConnection.h b/Engine/Plugins/Messaging/UdpMessaging/Source/UdpMessaging/Private/Tunnel/UdpMessageTunnelConnection.h
--- a/Engine/Plugins/Messaging/UdpMessaging/Source/UdpMessaging/Private/Tunnel/UdpMessageTunnelConnection.h
+++ b/Engine/Plugins/Messaging/UdpMessaging/Source/UdpMessaging/Private/Tunnel/UdpMessageTunnelConnection.h
@@ -52,10 +52,8 @@ public:
 		Thread = FRunnableThread::Create(this, TEXT("FUdpMessageTunnelConnection"), 128 * 1024, TPri_Normal);
 	}
 
-	/**
-	 * Destructor.
-	 */
-	~FUdpMessageTunnelConnection( )
+	/** Destructor. */
+	~FUdpMessageTunnelConnection()
 	{
 		if (Thread != nullptr)
 		{
@@ -99,12 +97,12 @@ public:
 
 	// FRunnable interface
 
-	virtual bool Init( ) override
+	virtual bool Init() override
 	{
 		return (Socket != nullptr);
 	}
 
-	virtual uint32 Run( ) override
+	virtual uint32 Run() override
 	{
 		bRun = true;
 
@@ -124,39 +122,39 @@ public:
 		return 0;
 	}
 
-	virtual void Stop( ) override
+	virtual void Stop() override
 	{
 		Socket->Close();
 		bRun = false;
 	}
 
-	virtual void Exit( ) override { }
+	virtual void Exit() override { }
 
 public:
 
 	// IUdpMessageTunnelConnection interface
 
-	virtual void Close( ) override
+	virtual void Close() override
 	{
 		Stop();
 	}
 
-	virtual uint64 GetTotalBytesReceived( ) const override
+	virtual uint64 GetTotalBytesReceived() const override
 	{
 		return TotalBytesReceived;
 	}
 
-	virtual uint64 GetTotalBytesSent( ) const override
+	virtual uint64 GetTotalBytesSent() const override
 	{
 		return TotalBytesSent;
 	}
 
-	virtual FText GetName( ) const
+	virtual FText GetName() const
 	{
 		return RemoteEndpoint.ToText();
 	}
 
-	virtual FTimespan GetUptime( ) const override
+	virtual FTimespan GetUptime() const override
 	{
 		if (IsOpen())
 		{
@@ -166,7 +164,7 @@ public:
 		return (ClosedTime - OpenedTime);
 	}
 
-	virtual bool IsOpen( ) const override
+	virtual bool IsOpen() const override
 	{
 		return (Socket->GetConnectionState() == SCS_Connected);
 	}
@@ -178,7 +176,7 @@ protected:
 	 *
 	 * @return true on success, false otherwise.
 	 */
-	bool ReceivePayloads( )
+	bool ReceivePayloads()
 	{
 		uint32 PendingDataSize = 0;
 
@@ -239,7 +237,7 @@ protected:
 	 *
 	 * @return true on success, false otherwise.
 	 */
-	bool SendPayloads( )
+	bool SendPayloads()
 	{
 		if (!Outbox.IsEmpty())
 		{
@@ -276,33 +274,33 @@ protected:
 
 private:
 
-	// Holds the time at which the connection was closed.
+	/** Holds the time at which the connection was closed. */
 	FDateTime ClosedTime;
 
-	// Holds the collection of received payloads.
+	/** Holds the collection of received payloads. */
 	TQueue<FArrayReaderPtr> Inbox;
 
-	// Holds the time at which the connection was opened.
+	/** Holds the time at which the connection was opened. */
 	FDateTime OpenedTime;
 
-	// Holds the collection of outbound payloads.
+	/** Holds the collection of outbound payloads. */
 	TQueue<FArrayReaderPtr> Outbox;
 
-	// Holds the IP endpoint of the remote client.
+	/** Holds the IP endpoint of the remote client. */
 	FIPv4Endpoint RemoteEndpoint;
 
-	// Holds the connection socket.
+	/** Holds the connection socket. */
 	FSocket* Socket;
 
-	// Holds the thread object.
+	/** Holds the thread object. */
 	FRunnableThread* Thread;
 
-	// Holds the total number of bytes received from the connection.
+	/** Holds the total number of bytes received from the connection. */
 	uint64 TotalBytesReceived;
 
-	// Holds the total number of bytes sent to the connection.
+	/** Holds the total number of bytes sent to the connection. */
 	uint64 TotalBytesSent;
 
-	// thread should continue running
+	/** thread should continue running. */
 	bool bRun;
 };
diff --git a/Engine/Plugins/Messaging/UdpMessaging/Source/UdpMessaging/Private/UdpMessagingModule.cpp b/Engine/Plugins/Messaging/UdpMessaging/Source/UdpMessaging/Private/UdpMessagingModule.cpp
--- a/Engine/Plugins/Messaging/UdpMessaging/Source/UdpMessaging/Private/UdpMessagingModule.cpp
+++ b/Engine/Plugins/Messaging/UdpMessaging/Source/UdpMessaging/Private/UdpMessagingModule.cpp
@@ -16,7 +16,7 @@ public:
 
 	// IModuleInterface interface
 
-	virtual void StartupModule( ) override
+	virtual void StartupModule() override
 	{
 		if (!SupportsNetworkedTransport())
 		{
@@ -54,7 +54,7 @@ public:
 		RestartServices();
 	}
 
-	virtual void ShutdownModule( ) override
+	virtual void ShutdownModule() override
 	{
 		// unregister application events
 		FCoreDelegates::ApplicationHasReactivatedDelegate.RemoveAll(this);
@@ -73,7 +73,7 @@ public:
 		ShutdownTunnel();
 	}
 
-	virtual bool SupportsDynamicReloading( ) override
+	virtual bool SupportsDynamicReloading() override
 	{
 		return true;
 	}
@@ -81,7 +81,7 @@ public:
 protected:
 
 	/** Initializes the message bridge with the current settings. */
-	void InitializeBridge( )
+	void InitializeBridge()
 	{
 		ShutdownBridge();
 
@@ -134,7 +134,7 @@ protected:
 	}
 
 	/** Initializes the message tunnel with the current settings. */
-	void InitializeTunnel( )
+	void InitializeTunnel()
 	{
 		ShutdownTunnel();
 
@@ -186,7 +186,7 @@ protected:
 	}
 
 	/** Restarts the bridge and tunnel services. */
-	void RestartServices( )
+	void RestartServices()
 	{
 		const UUdpMessagingSettings& Settings = *GetDefault<UUdpMessagingSettings>();
 
@@ -221,7 +221,7 @@ protected:
 	 * @todo gmp: this should be moved into an Engine module, so it can be shared with other transports
 	 * @return true if networked transport is supported, false otherwise.
 	 */
-	bool SupportsNetworkedTransport( ) const
+	bool SupportsNetworkedTransport() const
 	{
 		if (FApp::IsGame())
 		{
@@ -251,7 +251,7 @@ protected:
 	}
 
 	/** Shuts down the message bridge. */
-	void ShutdownBridge( )
+	void ShutdownBridge()
 	{
 		if (MessageBridge.IsValid())
 		{
@@ -262,7 +262,7 @@ protected:
 	}
 
 	/** Shuts down the message tunnel. */
-	void ShutdownTunnel( )
+	void ShutdownTunnel()
 	{
 		if (MessageTunnel.IsValid())
 		{
@@ -273,21 +273,21 @@ protected:
 
 private:
 
-	// Callback for when an has been reactivated (i.e. return from sleep on iOS).
-	void HandleApplicationHasReactivated( )
+	/** Callback for when an has been reactivated (i.e. return from sleep on iOS). */
+	void HandleApplicationHasReactivated()
 	{
 		RestartServices();
 	}
 
-	// Callback for when the application will be deactivated (i.e. sleep on iOS).
-	void HandleApplicationWillDeactivate( )
+	/** Callback for when the application will be deactivated (i.e. sleep on iOS).*/
+	void HandleApplicationWillDeactivate()
 	{
 		ShutdownBridge();
 		ShutdownTunnel();
 	}
 
-	// Callback for when the settings were saved.
-	bool HandleSettingsSaved( )
+	/** Callback for when the settings were saved. */
+	bool HandleSettingsSaved()
 	{
 		RestartServices();
 
@@ -296,10 +296,10 @@ private:
 
 private:
 
-	// Holds the message bridge if present
+	/** Holds the message bridge if present. */
 	IMessageBridgePtr MessageBridge;
 
-	// Holds the message tunnel if present
+	/** Holds the message tunnel if present. */
 	IUdpMessageTunnelPtr MessageTunnel;
 };
 
