UdpMessaging: Renamed local variable for better readability

diff --git a/Engine/Plugins/Messaging/UdpMessaging/Source/UdpMessaging/Private/Transport/UdpMessageProcessor.cpp b/Engine/Plugins/Messaging/UdpMessaging/Source/UdpMessaging/Private/Transport/UdpMessageProcessor.cpp
--- a/Engine/Plugins/Messaging/UdpMessaging/Source/UdpMessaging/Private/Transport/UdpMessageProcessor.cpp
+++ b/Engine/Plugins/Messaging/UdpMessaging/Source/UdpMessaging/Private/Transport/UdpMessageProcessor.cpp
@@ -244,8 +244,8 @@ void FUdpMessageProcessor::ConsumeOutboundMessages()
 
 			for (TMap<FIPv4Endpoint, FNodeInfo>::TIterator It(StaticNodes); It; ++It)
 			{
-				FNodeInfo& NodeInfo = It.Value();
-				NodeInfo.Segmenters.Add(LastSentMessage, MakeShareable(new FUdpMessageSegmenter(OutboundMessage.MessageData.ToSharedRef(), 1024)));
+				FNodeInfo& StaticNodeInfo = It.Value();
+				StaticNodeInfo.Segmenters.Add(LastSentMessage, MakeShareable(new FUdpMessageSegmenter(OutboundMessage.MessageData.ToSharedRef(), 1024)));
 			}
 		}
 
