Messaging: Code and documentation cleanup pass

diff --git a/Engine/Plugins/Messaging/UdpMessaging/Source/UdpMessaging/Public/IUdpMessageTunnelConnection.h b/Engine/Plugins/Messaging/UdpMessaging/Source/UdpMessaging/Public/IUdpMessageTunnelConnection.h
new file mode 100644
--- /dev/null
+++ b/Engine/Plugins/Messaging/UdpMessaging/Source/UdpMessaging/Public/IUdpMessageTunnelConnection.h
@@ -0,0 +1,62 @@
+// Copyright 1998-2014 Epic Games, Inc. All Rights Reserved.
+
+#pragma once
+
+
+/**
+ * Interface for message tunnel connections.
+ */
+class IUdpMessageTunnelConnection
+{
+public:
+
+	/** Closes this connection. */
+	virtual void Close() = 0;
+
+	/**
+	 * Gets the total number of bytes received from this connection.
+	 *
+	 * @return Number of bytes.
+	 */
+	virtual uint64 GetTotalBytesReceived() const = 0;
+
+	/**
+	 * Gets the total number of bytes received from this connection.
+	 *
+	 * @return Number of bytes.
+	 */
+	virtual uint64 GetTotalBytesSent() const = 0;
+
+	/**
+	 * Gets the human readable name of the connection.
+	 *
+	 * @return Connection name.
+	 */
+	virtual FText GetName() const = 0;
+
+	/**
+	 * Gets the amount of time that the connection has been established.
+	 *
+	 * @return Time span.
+	 */
+	virtual FTimespan GetUptime() const = 0;
+
+	/**
+	 * Checks whether this connection is open.
+	 *
+	 * @return true if the connection is open, false otherwise.
+	 */
+	virtual bool IsOpen() const = 0;
+
+public:
+
+	/** Virtual destructor. */
+	virtual ~IUdpMessageTunnelConnection() { }
+};
+
+
+/** Type definition for shared pointers to instances of IUdpMessageTunnelConnection. */
+typedef TSharedPtr<IUdpMessageTunnelConnection> IUdpMessageTunnelConnectionPtr;
+
+/** Type definition for shared references to instances of IUdpMessageTunnelConnection. */
+typedef TSharedRef<IUdpMessageTunnelConnection> IUdpMessageTunnelConnectionRef;
diff --git a/Engine/Plugins/Messaging/UdpMessaging/Source/UdpMessaging/Public/Interfaces/IUdpMessageTunnelConnection.h b/Engine/Plugins/Messaging/UdpMessaging/Source/UdpMessaging/Public/Interfaces/IUdpMessageTunnelConnection.h
deleted file mode 100644
diff --git a/Engine/Source/Runtime/CoreUObject/Private/Misc/MessageTypeMap.cpp b/Engine/Source/Runtime/CoreUObject/Private/Misc/MessageTypeMap.cpp
--- a/Engine/Source/Runtime/CoreUObject/Private/Misc/MessageTypeMap.cpp
+++ b/Engine/Source/Runtime/CoreUObject/Private/Misc/MessageTypeMap.cpp
@@ -1,11 +1,6 @@
 // Copyright 1998-2014 Epic Games, Inc. All Rights Reserved.
 
-/*=============================================================================
-	UnMisc.cpp: Various core platform-independent functions.
-=============================================================================*/
-
-// Core includes.
 #include "CoreUObjectPrivate.h"
 
-FMessageTypeMap FMessageTypeMap::MessageTypeMap;
 
+FMessageTypeMap FMessageTypeMap::MessageTypeMap;
diff --git a/Engine/Source/Runtime/CoreUObject/Public/Misc/MessageTypeMap.h b/Engine/Source/Runtime/CoreUObject/Public/Misc/MessageTypeMap.h
--- a/Engine/Source/Runtime/CoreUObject/Public/Misc/MessageTypeMap.h
+++ b/Engine/Source/Runtime/CoreUObject/Public/Misc/MessageTypeMap.h
@@ -1,14 +1,10 @@
 // Copyright 1998-2014 Epic Games, Inc. All Rights Reserved.
 
-/*=============================================================================
-	MessageTypeMap.h: Declares the FMessageTypeMap class.
-=============================================================================*/
-
 #pragma once
 
 
 /**
- * Implements a type map for messages
+ * Implements a type map for messages.
  */
 class FMessageTypeMap
 {
diff --git a/Engine/Source/Runtime/Messaging/Public/Interfaces/IAuthorizeMessageRecipients.h b/Engine/Source/Runtime/Messaging/Public/Interfaces/IAuthorizeMessageRecipients.h
--- a/Engine/Source/Runtime/Messaging/Public/Interfaces/IAuthorizeMessageRecipients.h
+++ b/Engine/Source/Runtime/Messaging/Public/Interfaces/IAuthorizeMessageRecipients.h
@@ -3,6 +3,7 @@
 #pragma once
 
 
+// forward declarations
 struct FMessageAddress;
 class IMessageInterceptor;
 class IReceiveMessages;
diff --git a/Engine/Source/Runtime/Messaging/Public/Interfaces/IMessageInterceptor.h b/Engine/Source/Runtime/Messaging/Public/Interfaces/IMessageInterceptor.h
--- a/Engine/Source/Runtime/Messaging/Public/Interfaces/IMessageInterceptor.h
+++ b/Engine/Source/Runtime/Messaging/Public/Interfaces/IMessageInterceptor.h
@@ -3,6 +3,7 @@
 #pragma once
 
 
+// forward declarations
 class IMessageContext;
 
 
diff --git a/Engine/Source/Runtime/Messaging/Public/Interfaces/IMessageSubscription.h b/Engine/Source/Runtime/Messaging/Public/Interfaces/IMessageSubscription.h
--- a/Engine/Source/Runtime/Messaging/Public/Interfaces/IMessageSubscription.h
+++ b/Engine/Source/Runtime/Messaging/Public/Interfaces/IMessageSubscription.h
@@ -3,6 +3,7 @@
 #pragma once
 
 
+// forward declarations
 enum class EMessageScope;
 class IReceiveMessages;
 
diff --git a/Engine/Source/Runtime/Messaging/Public/Interfaces/IMutableMessageContext.h b/Engine/Source/Runtime/Messaging/Public/Interfaces/IMutableMessageContext.h
--- a/Engine/Source/Runtime/Messaging/Public/Interfaces/IMutableMessageContext.h
+++ b/Engine/Source/Runtime/Messaging/Public/Interfaces/IMutableMessageContext.h
@@ -5,6 +5,7 @@
 #include "IMessageContext.h"
 
 
+// forward declarations
 class IMessageAttachment;
 class UScriptStruct;
 
diff --git a/Engine/Source/Runtime/Messaging/Public/Interfaces/IReceiveMessages.h b/Engine/Source/Runtime/Messaging/Public/Interfaces/IReceiveMessages.h
--- a/Engine/Source/Runtime/Messaging/Public/Interfaces/IReceiveMessages.h
+++ b/Engine/Source/Runtime/Messaging/Public/Interfaces/IReceiveMessages.h
@@ -3,6 +3,8 @@
 #pragma once
 
 
+// forward declarations
+enum ENamedThreads::Type;
 class IMessageContext;
 
 
diff --git a/Engine/Source/Runtime/Messaging/Public/Interfaces/ISendMessages.h b/Engine/Source/Runtime/Messaging/Public/Interfaces/ISendMessages.h
--- a/Engine/Source/Runtime/Messaging/Public/Interfaces/ISendMessages.h
+++ b/Engine/Source/Runtime/Messaging/Public/Interfaces/ISendMessages.h
@@ -3,6 +3,7 @@
 #pragma once
 
 
+// forward declarations
 class IMessageContext;
 
 
