Messaging: Various API and stability improvements

#UpgradeNotes:
- UdpMessaging protocol version is now 10 - rebuild all the things!
- IMessageBus::Forward() no longer takes a Scope parameter

Other changes:
- moved message serialization out of Messaging and into transport layer
- replaced legacy json message serializer with Serialization framework
- message serialization now more gracefully handles corrupted/malicious data
- first iteration towards "include what you use"
- moved several inline implementations into .cpp files
- documentation updates

In case of problems due to this changelist, throw your hands in the air and start screaming.

diff --git a/Engine/Plugins/Messaging/UdpMessaging/Source/UdpMessaging/Private/Transport/ReassembledUdpMessage.h b/Engine/Plugins/Messaging/UdpMessaging/Source/UdpMessaging/Private/Transport/ReassembledUdpMessage.h
deleted file mode 100644
diff --git a/Engine/Plugins/Messaging/UdpMessaging/Source/UdpMessaging/Private/Transport/UdpDeserializedMessage.cpp b/Engine/Plugins/Messaging/UdpMessaging/Source/UdpMessaging/Private/Transport/UdpDeserializedMessage.cpp
new file mode 100644
--- /dev/null
+++ b/Engine/Plugins/Messaging/UdpMessaging/Source/UdpMessaging/Private/Transport/UdpDeserializedMessage.cpp
@@ -0,0 +1,176 @@
+// Copyright 1998-2014 Epic Games, Inc. All Rights Reserved.
+
+#include "UdpMessagingPrivatePCH.h"
+#include "JsonStructDeserializerBackend.h"
+#include "StructDeserializer.h"
+
+
+/* FUdpDeserializedMessage interface
+ *****************************************************************************/
+
+bool FUdpDeserializedMessage::Deserialize( const FUdpReassembledMessageRef& ReassembledMessage )
+{
+	// Note that some complex values are deserialized manually here, so that we
+	// can sanity check their values. @see FUdpSerializeMessageTask::DoTask()
+
+	FMemoryReader MessageReader(ReassembledMessage->GetData());
+	MessageReader.ArMaxSerializeSize = NAME_SIZE;
+
+	// message type info
+	{
+		FName MessageType;
+		MessageReader << MessageType;
+
+		TypeInfo = FMessageTypeMap::MessageTypeMap.Find(MessageType.ToString());
+
+		if (!TypeInfo.IsValid(false, true))
+		{
+			return false;
+		}
+	}
+
+	// sender address
+	{
+		MessageReader << Sender;
+	}
+
+	// recipient addresses
+	{
+		int32 NumRecipients = 0;
+		MessageReader << NumRecipients;
+
+		if (NumRecipients > UDP_MESSAGING_MAX_RECIPIENTS)
+		{
+			return false;
+		}
+
+		Recipients.Empty(NumRecipients);
+
+		while (0 < NumRecipients--)
+		{
+			MessageReader << *::new(Recipients) FMessageAddress;
+		}
+	}
+
+	// message scope
+	{
+		MessageReader << Scope;
+
+		if (static_cast<uint8>(Scope.GetValue()) > static_cast<uint8>(EMessageScope::All))
+		{
+			return false;
+		}
+	}
+
+	// time sent & expiration
+	{
+		MessageReader << TimeSent;
+		MessageReader << Expiration;
+	}
+
+	// annotations
+	{
+		int32 NumAnnotations = 0;
+		MessageReader << NumAnnotations;
+
+		if (NumAnnotations > UDP_MESSAGING_MAX_ANNOTATIONS)
+		{
+			return false;
+		}
+
+		while (0 < NumAnnotations--)
+		{
+			FName Key;
+			FString Value;
+
+			MessageReader << Key;
+			MessageReader << Value;
+
+			Annotations.Add(Key, Value);
+		}
+	}
+
+	// create message body
+	MessageData = FMemory::Malloc(TypeInfo->PropertiesSize);
+	TypeInfo->InitializeScriptStruct(MessageData);
+
+	// deserialize message body
+	FJsonStructDeserializerBackend Backend(MessageReader);
+
+	return FStructDeserializer::Deserialize(MessageData, *TypeInfo, Backend);
+}
+
+
+/* IMessageContext interface
+ *****************************************************************************/
+
+const TMap<FName, FString>& FUdpDeserializedMessage::GetAnnotations() const
+{
+	return Annotations;
+}
+
+
+IMessageAttachmentPtr FUdpDeserializedMessage::GetAttachment() const
+{
+	return Attachment;
+}
+
+
+const FDateTime& FUdpDeserializedMessage::GetExpiration() const
+{
+	return Expiration;
+}
+
+
+const void* FUdpDeserializedMessage::GetMessage() const
+{
+	return MessageData;
+}
+
+
+const TWeakObjectPtr<UScriptStruct>& FUdpDeserializedMessage::GetMessageTypeInfo() const
+{
+	return TypeInfo;
+}
+
+
+IMessageContextPtr FUdpDeserializedMessage::GetOriginalContext() const
+{
+	return nullptr;
+}
+
+
+const TArray<FMessageAddress>& FUdpDeserializedMessage::GetRecipients() const
+{
+	return Recipients;
+}
+
+
+EMessageScope FUdpDeserializedMessage::GetScope() const
+{
+	return Scope;
+}
+
+
+const FMessageAddress& FUdpDeserializedMessage::GetSender() const
+{
+	return Sender;
+}
+
+
+ENamedThreads::Type FUdpDeserializedMessage::GetSenderThread() const
+{
+	return ENamedThreads::AnyThread;
+}
+
+
+const FDateTime& FUdpDeserializedMessage::GetTimeForwarded() const
+{
+	return TimeSent;
+}
+
+
+const FDateTime& FUdpDeserializedMessage::GetTimeSent() const
+{
+	return TimeSent;
+}
diff --git a/Engine/Plugins/Messaging/UdpMessaging/Source/UdpMessaging/Private/Transport/UdpDeserializedMessage.h b/Engine/Plugins/Messaging/UdpMessaging/Source/UdpMessaging/Private/Transport/UdpDeserializedMessage.h
new file mode 100644
--- /dev/null
+++ b/Engine/Plugins/Messaging/UdpMessaging/Source/UdpMessaging/Private/Transport/UdpDeserializedMessage.h
@@ -0,0 +1,85 @@
+// Copyright 1998-2014 Epic Games, Inc. All Rights Reserved.
+
+#pragma once
+
+
+/**
+ * Holds a deserialized message.
+ */
+class FUdpDeserializedMessage
+	: public IMessageContext
+{
+public:
+
+	/**
+	 * Creates and initializes a new instance.
+	 *
+	 * @param InAttachment An optional message attachment.
+	 */
+	FUdpDeserializedMessage( const IMessageAttachmentPtr& InAttachment )
+		: Attachment(InAttachment)
+	{ }
+
+public:
+
+	/**
+	 * Deserializes the given reassembled message.
+	 *
+	 * @param ReassembledMessage The reassembled message to deserialize.
+	 * @return true on success, false otherwise.
+	 */
+	bool Deserialize( const FUdpReassembledMessageRef& ReassembledMessage );
+
+public:
+
+	// IMessageContext interface
+
+	virtual const TMap<FName, FString>& GetAnnotations() const override;
+	virtual IMessageAttachmentPtr GetAttachment() const override;
+	virtual const FDateTime& GetExpiration() const override;
+	virtual const void* GetMessage() const override;
+	virtual const TWeakObjectPtr<UScriptStruct>& GetMessageTypeInfo() const override;
+	virtual IMessageContextPtr GetOriginalContext() const override;
+	virtual const TArray<FMessageAddress>& GetRecipients() const override;
+	virtual EMessageScope GetScope() const override;
+	virtual const FMessageAddress& GetSender() const override;
+	virtual ENamedThreads::Type GetSenderThread() const override;
+	virtual const FDateTime& GetTimeForwarded() const override;
+	virtual const FDateTime& GetTimeSent() const override;
+
+private:
+
+	/** Holds the optional message annotations. */
+	TMap<FName, FString> Annotations;
+
+	/** Holds a pointer to attached binary data. */
+	IMessageAttachmentPtr Attachment;
+
+	/** Holds the expiration time. */
+	FDateTime Expiration;
+
+	/** Holds the message. */
+	void* MessageData;
+
+	/** Holds the message recipients. */
+	TArray<FMessageAddress> Recipients;
+
+	/** Holds the message's scope. */
+	TEnumAsByte<EMessageScope> Scope;
+
+	/** Holds the sender's identifier. */
+	FMessageAddress Sender;
+
+	/** Holds the time at which the message was sent. */
+	FDateTime TimeSent;
+
+	/** Holds the message's type information. */
+	TWeakObjectPtr<UScriptStruct> TypeInfo;
+};
+
+
+/** Type definition for shared pointers to instances of FUdpDeserializedMessage. */
+typedef TSharedPtr<FUdpDeserializedMessage, ESPMode::ThreadSafe> FUdpDeserializedMessagePtr;
+
+/** Type definition for shared references to instances of FUdpDeserializedMessage. */
+typedef TSharedRef<FUdpDeserializedMessage, ESPMode::ThreadSafe> FUdpDeserializedMessageRef;
diff --git a/Engine/Plugins/Messaging/UdpMessaging/Source/UdpMessaging/Private/Transport/UdpMessageProcessor.cpp b/Engine/Plugins/Messaging/UdpMessaging/Source/UdpMessaging/Private/Transport/UdpMessageProcessor.cpp
--- a/Engine/Plugins/Messaging/UdpMessaging/Source/UdpMessaging/Private/Transport/UdpMessageProcessor.cpp
+++ b/Engine/Plugins/Messaging/UdpMessaging/Source/UdpMessaging/Private/Transport/UdpMessageProcessor.cpp
@@ -68,16 +68,16 @@ bool FUdpMessageProcessor::EnqueueInboundSegment( const FArrayReaderPtr& Data, c
 }
 
 
-bool FUdpMessageProcessor::EnqueueOutboundMessage( const IMessageDataRef& Data, const FGuid& Recipient )
+bool FUdpMessageProcessor::EnqueueOutboundMessage( const FUdpSerializedMessageRef& SerializedMessage, const FGuid& Recipient )
 {
-	if (!OutboundMessages.Enqueue(FOutboundMessage(Data, Recipient)))
+	if (!OutboundMessages.Enqueue(FOutboundMessage(SerializedMessage, Recipient)))
 	{
 		return false;
 	}
 
-	Data->OnStateChanged().BindRaw(this, &FUdpMessageProcessor::HandleMessageDataStateChanged);
+	SerializedMessage->OnStateChanged().BindRaw(this, &FUdpMessageProcessor::HandleSerializedMessageStateChanged);
 
-	if (Data->GetState() != EMessageDataState::Incomplete)
+	if (SerializedMessage->GetState() != EUdpSerializedMessageState::Incomplete)
 	{
 		WorkEvent->Trigger();
 	}
@@ -245,11 +245,11 @@ void FUdpMessageProcessor::ConsumeOutboundMessages()
 			for (TMap<FIPv4Endpoint, FNodeInfo>::TIterator It(StaticNodes); It; ++It)
 			{
 				FNodeInfo& StaticNodeInfo = It.Value();
-				StaticNodeInfo.Segmenters.Add(LastSentMessage, MakeShareable(new FUdpMessageSegmenter(OutboundMessage.MessageData.ToSharedRef(), 1024)));
+				StaticNodeInfo.Segmenters.Add(LastSentMessage, MakeShareable(new FUdpMessageSegmenter(OutboundMessage.SerializedMessage.ToSharedRef(), 1024)));
 			}
 		}
 
-		RecipientNodeInfo.Segmenters.Add(LastSentMessage, MakeShareable(new FUdpMessageSegmenter(OutboundMessage.MessageData.ToSharedRef(), 1024)));
+		RecipientNodeInfo.Segmenters.Add(LastSentMessage, MakeShareable(new FUdpMessageSegmenter(OutboundMessage.SerializedMessage.ToSharedRef(), 1024)));
 	}
 }
 
@@ -320,40 +320,37 @@ void FUdpMessageProcessor::ProcessDataSegment( FInboundSegment& Segment, FNodeIn
 		return;
 	}
 
-	FReassembledUdpMessagePtr& Message = NodeInfo.ReassembledMessages.FindOrAdd(DataChunk.MessageId);
+	FUdpReassembledMessagePtr& ReassembledMessage = NodeInfo.ReassembledMessages.FindOrAdd(DataChunk.MessageId);
 
 	// Reassemble message
-	if (!Message.IsValid())
+	if (!ReassembledMessage.IsValid())
 	{
-		Message = MakeShareable(new FReassembledUdpMessage(DataChunk.MessageSize, DataChunk.TotalSegments, DataChunk.Sequence, Segment.Sender));
+		ReassembledMessage = MakeShareable(new FReassembledUdpMessage(DataChunk.MessageSize, DataChunk.TotalSegments, DataChunk.Sequence, Segment.Sender));
 	}
 
-	Message->Reassemble(DataChunk.SegmentNumber, DataChunk.SegmentOffset, DataChunk.Data, CurrentTime);
+	ReassembledMessage->Reassemble(DataChunk.SegmentNumber, DataChunk.SegmentOffset, DataChunk.Data, CurrentTime);
 
 	// Deliver or re-sequence message
-	if (Message->IsComplete())
+	if (ReassembledMessage->IsComplete())
 	{
 		AcknowledgeReceipt(DataChunk.MessageId, NodeInfo);
 
-		if (Message->GetSequence() == 0)
+		if (ReassembledMessage->GetSequence() == 0)
 		{
-			FMemoryReader HackReader(Message->GetData());
-
 			if (NodeInfo.NodeId.IsValid())
 			{
-				MessageReceivedDelegate.ExecuteIfBound(HackReader, nullptr, NodeInfo.NodeId);
+				MessageReassembledDelegate.ExecuteIfBound(ReassembledMessage.ToSharedRef(), nullptr, NodeInfo.NodeId);
 			}
 		}
-		else if (NodeInfo.Resequencer.Resequence(Message))
+		else if (NodeInfo.Resequencer.Resequence(ReassembledMessage))
 		{
-			FReassembledUdpMessagePtr ResequencedMessage;
-			FMemoryReader HackReader(ResequencedMessage->GetData());
+			FUdpReassembledMessagePtr ResequencedMessage;
 
 			while (NodeInfo.Resequencer.Pop(ResequencedMessage))
 			{
 				if (NodeInfo.NodeId.IsValid())
 				{
-					MessageReceivedDelegate.ExecuteIfBound(HackReader, nullptr, NodeInfo.NodeId);
+					MessageReassembledDelegate.ExecuteIfBound(ResequencedMessage.ToSharedRef(), nullptr, NodeInfo.NodeId);
 				}
 			}
 		}
@@ -515,7 +512,7 @@ void FUdpMessageProcessor::UpdateStaticNodes()
 /* FUdpMessageProcessor callbacks
  *****************************************************************************/
 
-void FUdpMessageProcessor::HandleMessageDataStateChanged()
+void FUdpMessageProcessor::HandleSerializedMessageStateChanged()
 {
 	WorkEvent->Trigger();
 }
diff --git a/Engine/Plugins/Messaging/UdpMessaging/Source/UdpMessaging/Private/Transport/UdpMessageProcessor.h b/Engine/Plugins/Messaging/UdpMessaging/Source/UdpMessaging/Private/Transport/UdpMessageProcessor.h
--- a/Engine/Plugins/Messaging/UdpMessaging/Source/UdpMessaging/Private/Transport/UdpMessageProcessor.h
+++ b/Engine/Plugins/Messaging/UdpMessaging/Source/UdpMessaging/Private/Transport/UdpMessageProcessor.h
@@ -3,6 +3,16 @@
 #pragma once
 
 
+/**
+ * Delegate type for handling the message reassembly completion.
+ *
+ * The first parameter is the reassembled message data.
+ * The second parameter contains the optional message attachment.
+ * The third parameter is the identifier of the node that sent the message.
+ */
+DECLARE_DELEGATE_ThreeParams(FOnMessageProcessorMessageReassembled, const FUdpReassembledMessageRef& /*ReassembledMessage*/, const IMessageAttachmentPtr& /*Attachment*/, const FGuid& /*NodeId*/);
+
+
 /**
  * Implements a message processor for UDP messages.
  */
@@ -22,7 +32,7 @@ class FUdpMessageProcessor
 		FGuid NodeId;
 
 		// Holds the collection of reassembled messages.
-		TMap<int32, FReassembledUdpMessagePtr> ReassembledMessages;
+		TMap<int32, FUdpReassembledMessagePtr> ReassembledMessages;
 
 		// Holds the message resequencer.
 		FUdpMessageResequencer Resequencer;
@@ -74,8 +84,8 @@ class FUdpMessageProcessor
 	// Structure for outbound messages.
 	struct FOutboundMessage
 	{
-		// Holds the message.
-		IMessageDataPtr MessageData;
+		// Holds the serialized message.
+		FUdpSerializedMessagePtr SerializedMessage;
 
 		// Holds the recipient.
 		FGuid RecipientId;
@@ -85,8 +95,8 @@ class FUdpMessageProcessor
 		FOutboundMessage() { }
 
 		// Creates and initializes a new instance.
-		FOutboundMessage( const IMessageDataRef& InMessageData, const FGuid& InRecipientId )
-			: MessageData(InMessageData)
+		FOutboundMessage( const FUdpSerializedMessageRef& InSerializedMessage, const FGuid& InRecipientId )
+			: SerializedMessage(InSerializedMessage)
 			, RecipientId(InRecipientId)
 		{ }
 	};
@@ -119,14 +129,24 @@ public:
 	/**
 	 * Queues up an outbound message.
 	 *
-	 * @param Data The message data to send.
+	 * @param SerializedMessage The serialized message to send.
 	 * @param Recipient The recipient's IPv4 endpoint.
 	 * @return true if the message was queued up, false otherwise.
 	 */
-	bool EnqueueOutboundMessage( const IMessageDataRef& Data, const FGuid& Recipient );
+	bool EnqueueOutboundMessage( const FUdpSerializedMessageRef& SerializedMessage, const FGuid& Recipient );
 
 public:
 
+	/**
+	 * Returns a delegate that is executed when message data has been received.
+	 *
+	 * @param Delegate The delegate to set.
+	 */
+	FOnMessageProcessorMessageReassembled& OnMessageReassembled()
+	{
+		return MessageReassembledDelegate;
+	}
+
 	FOnMessageTransportNodeDiscovered& OnNodeDiscovered()
 	{
 		return NodeDiscoveredDelegate;
@@ -141,16 +161,6 @@ public:
 	{
 		return NodeLostDelegate;
 	}
-
-	/**
-	 * Returns a delegate that is executed when message data has been received.
-	 *
-	 * @param Delegate The delegate to set.
-	 */
-	FOnMessageTransportMessageReceived& OnMessageReceived()
-	{
-		return MessageReceivedDelegate;
-	}
 	
 public:
 
@@ -283,7 +293,7 @@ protected:
 private:
 
 	/** Handles message data state changes. */
-	void HandleMessageDataStateChanged();
+	void HandleSerializedMessageStateChanged();
 
 private:
 
@@ -334,7 +344,7 @@ private:
 private:
 
 	/** Holds a delegate to be invoked when a message was received on the transport channel. */
-	FOnMessageTransportMessageReceived MessageReceivedDelegate;
+	FOnMessageProcessorMessageReassembled MessageReassembledDelegate;
 
 	/** Holds a delegate to be invoked when a network node was discovered. */
 	FOnMessageTransportNodeDiscovered NodeDiscoveredDelegate;
diff --git a/Engine/Plugins/Messaging/UdpMessaging/Source/UdpMessaging/Private/Transport/UdpMessageResequencer.h b/Engine/Plugins/Messaging/UdpMessaging/Source/UdpMessaging/Private/Transport/UdpMessageResequencer.h
--- a/Engine/Plugins/Messaging/UdpMessaging/Source/UdpMessaging/Private/Transport/UdpMessageResequencer.h
+++ b/Engine/Plugins/Messaging/UdpMessaging/Source/UdpMessaging/Private/Transport/UdpMessageResequencer.h
@@ -38,7 +38,7 @@ public:
 	 *
 	 * @return true if a message was returned, false otherwise.
 	 */
-	bool Pop(FReassembledUdpMessagePtr& OutMessage)
+	bool Pop(FUdpReassembledMessagePtr& OutMessage)
 	{
 		if (MessageHeap.HeapTop()->GetSequence() == NextSequence)
 		{
@@ -58,7 +58,7 @@ public:
 	 * @param Message The message to resequence.
 	 * @return true if the message is in sequence, false otherwise.
 	 */
-	bool Resequence(const FReassembledUdpMessagePtr& Message)
+	bool Resequence(const FUdpReassembledMessagePtr& Message)
 	{
 		MessageHeap.HeapPush(Message, FSequenceComparer());
 
@@ -78,7 +78,7 @@ private:
 	/** Helper for ordering messages by their sequence numbers. */
 	struct FSequenceComparer
 	{
-		bool operator () (const FReassembledUdpMessagePtr& A, const FReassembledUdpMessagePtr& B) const
+		bool operator()(const FUdpReassembledMessagePtr& A, const FUdpReassembledMessagePtr& B) const
 		{
 			return A->GetSequence() < B->GetSequence();
 		}
@@ -93,7 +93,7 @@ private:
 	uint64 HighestReceivedSequence;
 
 	/** Holds the messages that need to be resequenced. */
-	TArray<FReassembledUdpMessagePtr> MessageHeap;
+	TArray<FUdpReassembledMessagePtr> MessageHeap;
 
 	/** Holds the maximum resequence window size. */
 	uint16 WindowSize;
diff --git a/Engine/Plugins/Messaging/UdpMessaging/Source/UdpMessaging/Private/Transport/UdpMessageSegmenter.cpp b/Engine/Plugins/Messaging/UdpMessaging/Source/UdpMessaging/Private/Transport/UdpMessageSegmenter.cpp
new file mode 100644
--- /dev/null
+++ b/Engine/Plugins/Messaging/UdpMessaging/Source/UdpMessaging/Private/Transport/UdpMessageSegmenter.cpp
@@ -0,0 +1,103 @@
+// Copyright 1998-2014 Epic Games, Inc. All Rights Reserved.
+
+#include "UdpMessagingPrivatePCH.h"
+
+
+/* FUdpMessageSegmenter structors
+ *****************************************************************************/
+
+FUdpMessageSegmenter::~FUdpMessageSegmenter()
+{
+	if (MessageReader != nullptr)
+	{
+		delete MessageReader;
+	}
+}
+
+
+/* FUdpMessageSegmenter interface
+ *****************************************************************************/
+
+int64 FUdpMessageSegmenter::GetMessageSize() const
+{
+	if (MessageReader == nullptr)
+	{
+		return 0;
+	}
+
+	return MessageReader->TotalSize();
+}
+
+
+bool FUdpMessageSegmenter::GetNextPendingSegment(TArray<uint8>& OutData, uint16& OutSegment) const
+{
+	if (MessageReader == nullptr)
+	{
+		return false;
+	}
+
+	for (TConstSetBitIterator<> It(PendingSegments); It; ++It)
+	{
+		OutSegment = It.GetIndex();
+
+		uint32 SegmentOffset = OutSegment * SegmentSize;
+		int32 ActualSegmentSize = MessageReader->TotalSize() - SegmentOffset;
+
+		if (ActualSegmentSize > SegmentSize)
+		{
+			ActualSegmentSize = SegmentSize;
+		}
+
+		OutData.Reset(ActualSegmentSize);
+		OutData.AddUninitialized(ActualSegmentSize);
+
+		MessageReader->Seek(SegmentOffset);
+		MessageReader->Serialize(OutData.GetData(), ActualSegmentSize);
+
+		//FMemory::Memcpy(OutData.GetTypedData(), Message->GetTypedData() + SegmentOffset, ActualSegmentSize);
+
+		return true;
+	}
+
+	return false;
+}
+
+
+void FUdpMessageSegmenter::Initialize()
+{
+	if (MessageReader != nullptr)
+	{
+		return;
+	}
+
+	if (SerializedMessage->GetState() == EUdpSerializedMessageState::Complete)
+	{
+		MessageReader = SerializedMessage->CreateReader();
+		PendingSegmentsCount = (MessageReader->TotalSize() + SegmentSize - 1) / SegmentSize;
+		PendingSegments.Init(true, PendingSegmentsCount);
+	}
+}
+
+
+void FUdpMessageSegmenter::MarkAsSent(uint16 Segment)
+{
+	if (Segment < PendingSegments.Num())
+	{
+		PendingSegments[Segment] = false;
+		--PendingSegmentsCount;
+	}
+}
+
+
+void FUdpMessageSegmenter::MarkForRetransmission(const TArray<uint16>& Segments)
+{
+	for (int32 Index = 0; Index < Segments.Num(); ++Index)
+	{
+		uint16 Segment = Segments[Index];
+
+		if (Segment < PendingSegments.Num())
+		{
+			PendingSegments[Segment] = true;
+		}
+	}
+}
diff --git a/Engine/Plugins/Messaging/UdpMessaging/Source/UdpMessaging/Private/Transport/UdpMessageSegmenter.h b/Engine/Plugins/Messaging/UdpMessaging/Source/UdpMessaging/Private/Transport/UdpMessageSegmenter.h
--- a/Engine/Plugins/Messaging/UdpMessaging/Source/UdpMessaging/Private/Transport/UdpMessageSegmenter.h
+++ b/Engine/Plugins/Messaging/UdpMessaging/Source/UdpMessaging/Private/Transport/UdpMessageSegmenter.h
@@ -21,22 +21,16 @@ public:
 	/**
 	 * Creates and initializes a new message segmenter.
 	 *
-	 * @param InMessage The message to segment.
+	 * @param InMessage The serialized message to segment.
 	 */
-	FUdpMessageSegmenter(const IMessageDataRef& InMessage, uint16 InSegmentSize)
-		: Message(InMessage)
-		, MessageReader(nullptr)
+	FUdpMessageSegmenter(const FUdpSerializedMessageRef& InSerializedMessage, uint16 InSegmentSize)
+		: MessageReader(nullptr)
 		, SegmentSize(InSegmentSize)
+		, SerializedMessage(InSerializedMessage)
 	{ }
 
 	/** Destructor. */
-	~FUdpMessageSegmenter()
-	{
-		if (MessageReader != nullptr)
-		{
-			delete MessageReader;
-		}
-	}
+	~FUdpMessageSegmenter();
 
 public:
 
@@ -45,15 +39,7 @@ public:
 	 *
 	 * @return Message size.
 	 */
-	int64 GetMessageSize() const
-	{
-		if (MessageReader == nullptr)
-		{
-			return 0;
-		}
-
-		return MessageReader->TotalSize();
-	}
+	int64 GetMessageSize() const;
 
 	/**
 	 * Gets the next pending segment.
@@ -62,38 +48,7 @@ public:
 	 * @param OutSegment Will hold the segment number.
 	 * @return true if a segment was returned, false if there are no more pending segments.
 	 */
-	bool GetNextPendingSegment(TArray<uint8>& OutData, uint16& OutSegment) const
-	{
-		if (MessageReader == nullptr)
-		{
-			return false;
-		}
-
-		for (TConstSetBitIterator<> It(PendingSegments); It; ++It)
-		{
-			OutSegment = It.GetIndex();
-
-			uint32 SegmentOffset = OutSegment * SegmentSize;
-			int32 ActualSegmentSize = MessageReader->TotalSize() - SegmentOffset;
-
-			if (ActualSegmentSize > SegmentSize)
-			{
-				ActualSegmentSize = SegmentSize;
-			}
-
-			OutData.Reset(ActualSegmentSize);
-			OutData.AddUninitialized(ActualSegmentSize);
-
-			MessageReader->Seek(SegmentOffset);
-			MessageReader->Serialize(OutData.GetData(), ActualSegmentSize);
-
-			//FMemory::Memcpy(OutData.GetTypedData(), Message->GetTypedData() + SegmentOffset, ActualSegmentSize);
-
-			return true;
-		}
-
-		return false;
-	}
+	bool GetNextPendingSegment(TArray<uint8>& OutData, uint16& OutSegment) const;
 
 	/**
 	 * Gets the number of segments that haven't been received yet.
@@ -116,20 +71,7 @@ public:
 	}
 
 	/** Initializes the segmenter. */
-	void Initialize()
-	{
-		if (MessageReader != nullptr)
-		{
-			return;
-		}
-
-		if (Message->GetState() == EMessageDataState::Complete)
-		{
-			MessageReader = Message->CreateReader();
-			PendingSegmentsCount = (MessageReader->TotalSize() + SegmentSize - 1) / SegmentSize;
-			PendingSegments.Init(true, PendingSegmentsCount);
-		}
-	}
+	void Initialize();
 
 	/**
 	 * Checks whether all segments have been sent.
@@ -158,7 +100,7 @@ public:
 	 */
 	bool IsInvalid() const
 	{
-		return (Message->GetState() == EMessageDataState::Invalid);
+		return (SerializedMessage->GetState() == EUdpSerializedMessageState::Invalid);
 	}
 
 	/**
@@ -166,14 +108,7 @@ public:
 	 *
 	 * @param Segment The sent segment.
 	 */
-	void MarkAsSent(uint16 Segment)
-	{
-		if (Segment < PendingSegments.Num())
-		{
-			PendingSegments[Segment] = false;
-			--PendingSegmentsCount;
-		}
-	}
+	void MarkAsSent(uint16 Segment);
 
 	/**
 	 * Marks the entire message for retransmission.
@@ -188,24 +123,10 @@ public:
 	 *
 	 * @param Segments The data segments to retransmit.
 	 */
-	void MarkForRetransmission(const TArray<uint16>& Segments)
-	{
-		for (int32 Index = 0; Index < Segments.Num(); ++Index)
-		{
-			uint16 Segment = Segments[Index];
-
-			if (Segment < PendingSegments.Num())
-			{
-				PendingSegments[Segment] = true;
-			}
-		}
-	}
+	void MarkForRetransmission(const TArray<uint16>& Segments);
 
 private:
 
-	/** Holds the message. */
-	IMessageDataPtr Message;
-
 	/** temp hack to support new transport API. */
 	FArchive* MessageReader;
 
@@ -217,4 +138,7 @@ private:
 
 	/** Holds the segment size. */
 	uint16 SegmentSize;
+
+	/** Holds the message. */
+	FUdpSerializedMessagePtr SerializedMessage;
 };
diff --git a/Engine/Plugins/Messaging/UdpMessaging/Source/UdpMessaging/Private/Transport/UdpMessageTransport.cpp b/Engine/Plugins/Messaging/UdpMessaging/Source/UdpMessaging/Private/Transport/UdpMessageTransport.cpp
--- a/Engine/Plugins/Messaging/UdpMessaging/Source/UdpMessaging/Private/Transport/UdpMessageTransport.cpp
+++ b/Engine/Plugins/Messaging/UdpMessaging/Source/UdpMessaging/Private/Transport/UdpMessageTransport.cpp
@@ -1,6 +1,7 @@
 // Copyright 1998-2014 Epic Games, Inc. All Rights Reserved.
 
 #include "UdpMessagingPrivatePCH.h"
+#include "UdpSerializeMessageTask.h"
 
 
 /* FUdpMessageTransport structors
@@ -72,7 +73,7 @@ bool FUdpMessageTransport::StartTransport()
 
 	// initialize threads
 	MessageProcessor = new FUdpMessageProcessor(UnicastSocket, FGuid::NewGuid(), MulticastEndpoint);
-	MessageProcessor->OnMessageReceived().BindRaw(this, &FUdpMessageTransport::HandleProcessorMessageReceived);
+	MessageProcessor->OnMessageReassembled().BindRaw(this, &FUdpMessageTransport::HandleProcessorMessageReassembled);
 	MessageProcessor->OnNodeDiscovered().BindRaw(this, &FUdpMessageTransport::HandleProcessorNodeDiscovered);
 	MessageProcessor->OnNodeLost().BindRaw(this, &FUdpMessageTransport::HandleProcessorNodeLost);
 
@@ -113,23 +114,31 @@ void FUdpMessageTransport::StopTransport()
 }
 
 
-bool FUdpMessageTransport::TransportMessage( const IMessageDataRef& Data, const IMessageAttachmentPtr& Attachment, const TArray<FGuid>& Recipients )
+bool FUdpMessageTransport::TransportMessage( const IMessageContextRef& Context, const TArray<FGuid>& Recipients )
 {
 	if (MessageProcessor == nullptr)
 	{
 		return false;
 	}
 
-	// published message
+	if (Context->GetRecipients().Num() > UDP_MESSAGING_MAX_RECIPIENTS)
+	{
+		return false;
+	}
+
+	FUdpSerializedMessageRef SerializedMessage = MakeShareable(new FUdpSerializedMessage());
+	TGraphTask<FUdpSerializeMessageTask>::CreateTask().ConstructAndDispatchWhenReady(Context, SerializedMessage);
+
+	// publish the message
 	if (Recipients.Num() == 0)
 	{
-		return MessageProcessor->EnqueueOutboundMessage(Data, FGuid());
+		return MessageProcessor->EnqueueOutboundMessage(SerializedMessage, FGuid());
 	}
 
-	// sent message
+	// send the message
 	for (int32 Index = 0; Index < Recipients.Num(); ++Index)
 	{
-		if (!MessageProcessor->EnqueueOutboundMessage(Data, Recipients[Index]))
+		if (!MessageProcessor->EnqueueOutboundMessage(SerializedMessage, Recipients[Index]))
 		{
 			return false;
 		}
@@ -142,6 +151,18 @@ bool FUdpMessageTransport::TransportMessage( const IMessageDataRef& Data, const
 /* FUdpMessageTransport event handlers
  *****************************************************************************/
 
+void FUdpMessageTransport::HandleProcessorMessageReassembled( const FUdpReassembledMessageRef& ReassembledMessage, const IMessageAttachmentPtr& Attachment, const FGuid& NodeId )
+{
+	// @todo gmp: move message deserialization into an async task
+	FUdpDeserializedMessageRef DeserializedMessage = MakeShareable(new FUdpDeserializedMessage(Attachment));
+
+	if (DeserializedMessage->Deserialize(ReassembledMessage))
+	{
+		MessageReceivedDelegate.ExecuteIfBound(DeserializedMessage, NodeId);
+	}
+}
+
+
 void FUdpMessageTransport::HandleSocketDataReceived( const FArrayReaderPtr& Data, const FIPv4Endpoint& Sender )
 {
 	if (MessageProcessor != nullptr)
diff --git a/Engine/Plugins/Messaging/UdpMessaging/Source/UdpMessaging/Private/Transport/UdpMessageTransport.h b/Engine/Plugins/Messaging/UdpMessaging/Source/UdpMessaging/Private/Transport/UdpMessageTransport.h
--- a/Engine/Plugins/Messaging/UdpMessaging/Source/UdpMessaging/Private/Transport/UdpMessageTransport.h
+++ b/Engine/Plugins/Messaging/UdpMessaging/Source/UdpMessaging/Private/Transport/UdpMessageTransport.h
@@ -49,15 +49,12 @@ public:
 
 	virtual bool StartTransport() override;
 	virtual void StopTransport() override;
-	virtual bool TransportMessage( const IMessageDataRef& Data, const IMessageAttachmentPtr& Attachment, const TArray<FGuid>& Recipients ) override;
+	virtual bool TransportMessage( const IMessageContextRef& Context, const TArray<FGuid>& Recipients ) override;
 
 private:
 
 	/** Handles received transport messages. */
-	void HandleProcessorMessageReceived( FArchive& MessageData, const IMessageAttachmentPtr& Attachment, const FGuid& NodeId )
-	{
-		MessageReceivedDelegate.ExecuteIfBound(MessageData, Attachment, NodeId);
-	}
+	void HandleProcessorMessageReassembled( const FUdpReassembledMessageRef& ReassembledMessage, const IMessageAttachmentPtr& Attachment, const FGuid& NodeId );
 
 	/** Handles discovered transport endpoints. */
 	void HandleProcessorNodeDiscovered( const FGuid& DiscoveredNodeId )
diff --git a/Engine/Plugins/Messaging/UdpMessaging/Source/UdpMessaging/Private/Transport/UdpReassembledMessage.h b/Engine/Plugins/Messaging/UdpMessaging/Source/UdpMessaging/Private/Transport/UdpReassembledMessage.h
new file mode 100644
--- /dev/null
+++ b/Engine/Plugins/Messaging/UdpMessaging/Source/UdpMessaging/Private/Transport/UdpReassembledMessage.h
@@ -0,0 +1,193 @@
+// Copyright 1998-2014 Epic Games, Inc. All Rights Reserved.
+
+#pragma once
+
+
+/** Type definition for shared pointers to instances of FUdpMessage. */
+typedef TSharedPtr<class FReassembledUdpMessage, ESPMode::ThreadSafe> FUdpReassembledMessagePtr;
+
+/** Type definition for shared references to instances of FUdpMessage. */
+typedef TSharedRef<class FReassembledUdpMessage, ESPMode::ThreadSafe> FUdpReassembledMessageRef;
+
+
+/**
+ * Implements a reassembled message.
+ */
+class FReassembledUdpMessage
+{
+public:
+
+	/** Default constructor. */
+	FReassembledUdpMessage() { }
+
+	/**
+	 * Creates and initializes a new inbound message info.
+	 *
+	 * @param MessageSize The total size of the message in bytes.
+	 * @param SegmentCount The total number of segments that need to be received for this message.
+	 * @param InSequence The message sequence number.
+	 * @param InSender The IPv4 endpoint of the sender.
+	 */
+	FReassembledUdpMessage( int32 MessageSize, int32 SegmentCount, uint64 InSequence, const FIPv4Endpoint& InSender )
+		: PendingSegments(true, SegmentCount)
+		, PendingSegmentsCount(SegmentCount)
+		, ReceivedBytes(0)
+		, Sender(InSender)
+		, Sequence(InSequence)
+	{
+		Data.AddUninitialized(MessageSize);
+	}
+
+	/** Virtual destructor. */
+	virtual ~FReassembledUdpMessage() { }
+
+public:
+
+	/**
+	 * Gets the message data.
+	 *
+	 * @return Message data.
+	 */
+	virtual const TArray<uint8>& GetData() const
+	{
+		return Data;
+	}
+
+	/**
+	 * Gets the time at which the last segment was received.
+	 *
+	 * @return Last receive time.
+	 */
+	FDateTime GetLastSegmentTime() const
+	{
+		return LastSegmentTime;
+	}
+
+	/**
+	 * Gets the list of segments that haven't been received yet.
+	 *
+	 * @return List of pending segment numbers.
+	 */
+	TArray<uint16> GetPendingSegments() const
+	{
+		TArray<uint16> Result;
+
+		if (PendingSegmentsCount > 0)
+		{
+			for (TConstSetBitIterator<> It(PendingSegments); It; ++It)
+			{
+				Result.Add(It.GetIndex());
+			}
+		}
+
+		return Result;
+	}
+
+	/**
+	 * Gets the number of segments that haven't been received yet.
+	 *
+	 * @return Number of pending segments.
+	 */
+	uint16 GetPendingSegmentsCount() const
+	{
+		return PendingSegmentsCount;
+	}
+
+	/**
+	 * Gets the number of retransmit requests that were sent for this payload.
+	 *
+	 * @return Number of retransmit requests.
+	 */
+	int32 GetRetransmitRequestsCount() const
+	{
+		return RetransmitRequestsCount;
+	}
+
+	/**
+	 * Gets the message's sequence number.
+	 *
+	 * @return Sequence number.
+	 */
+	uint64 GetSequence() const
+	{
+		return Sequence;
+	}
+
+	/**
+	 * Checks whether this message is complete.
+	 *
+	 * @return true if the message is complete, false otherwise.
+	 */
+	bool IsComplete() const
+	{
+		return (PendingSegmentsCount == 0);
+	}
+
+	/**
+	 * Checks whether this message has been initialized.
+	 *
+	 * @return true if the message is initialized, false otherwise.
+	 */
+	bool IsInitialized() const
+	{
+		return (Data.Num() < 0);
+	}
+
+	/**
+	 * Reassembles a segment into the specified message.
+	 *
+	 * @param SegmentNumber The number of the message segment.
+	 * @param SegmentOffset The segment's offset within the message.
+	 * @param SegmentData The segment data.
+	 * @param CurrentTime The current time.
+	 */
+	void Reassemble( int32 SegmentNumber, int32 SegmentOffset, const TArray<uint8>& SegmentData, const FDateTime& CurrentTime )
+	{
+		if (SegmentNumber >= PendingSegments.Num())
+		{
+			// temp hack sanity check
+			return;
+		}
+
+		LastSegmentTime = CurrentTime;
+
+		if (PendingSegments[SegmentNumber])
+		{
+			if (SegmentOffset + SegmentData.Num() <= Data.Num())
+			{
+				FMemory::Memcpy(Data.GetData() + SegmentOffset, SegmentData.GetData(), SegmentData.Num());
+
+				PendingSegments[SegmentNumber] = false;
+				--PendingSegmentsCount;
+
+				ReceivedBytes += SegmentData.Num();
+			}
+		}
+	}
+
+private:
+
+	/** Holds the message data. */
+	TArray<uint8> Data;
+
+	/** Holds the time at which the last segment was received. */
+	FDateTime LastSegmentTime;
+
+	/** Holds an array of bits that indicate which segments still need to be received. */
+	TBitArray<> PendingSegments;
+
+	/** Holds the number of segments that haven't been received yet. */
+	uint16 PendingSegmentsCount;
+
+	/** Holds the number of bytes received so far. */
+	int32 ReceivedBytes;
+
+	/** Holds the number of retransmit requests that were sent since the last segment was received. */
+	int32 RetransmitRequestsCount;
+
+	/** Holds the sender. */
+	FIPv4Endpoint Sender;
+
+	/** Holds the message sequence. */
+	uint64 Sequence;
+};
diff --git a/Engine/Plugins/Messaging/UdpMessaging/Source/UdpMessaging/Private/Transport/UdpSerializeMessageTask.cpp b/Engine/Plugins/Messaging/UdpMessaging/Source/UdpMessaging/Private/Transport/UdpSerializeMessageTask.cpp
new file mode 100644
--- /dev/null
+++ b/Engine/Plugins/Messaging/UdpMessaging/Source/UdpMessaging/Private/Transport/UdpSerializeMessageTask.cpp
@@ -0,0 +1,80 @@
+// Copyright 1998-2014 Epic Games, Inc. All Rights Reserved.
+
+#include "UdpMessagingPrivatePCH.h"
+#include "TaskGraphInterfaces.h"
+
+
+/* FUdpSerializeMessageTask interface
+ *****************************************************************************/
+
+void FUdpSerializeMessageTask::DoTask( ENamedThreads::Type CurrentThread, const FGraphEventRef& MyCompletionGraphEvent )
+{
+	if (MessageContext->IsValid())
+	{
+		// Note that some complex values are serialized manually here, so that we can ensure
+		// a consistent wire format, if their implementations change. This allows us to sanity
+		// check the values during deserialization. @see FUdpDeserializeMessage::Deserialize()
+
+		FArchive& Archive = SerializedMessage.Get();
+
+		// serialize context
+		{
+			const FName& MessageType = MessageContext->GetMessageType();
+			Archive << const_cast<FName&>(MessageType);
+
+			const FMessageAddress& Sender = MessageContext->GetSender();
+			Archive << const_cast<FMessageAddress&>(Sender);
+
+			const TArray<FMessageAddress>& Recipients = MessageContext->GetRecipients();
+			Archive << const_cast<TArray<FMessageAddress>&>(Recipients);
+
+			TEnumAsByte<EMessageScope> Scope = MessageContext->GetScope();
+			Archive << Scope;
+
+			const FDateTime& TimeSent = MessageContext->GetTimeSent();
+			Archive << const_cast<FDateTime&>(TimeSent);
+
+			const FDateTime& Expiration = MessageContext->GetExpiration();
+			Archive << const_cast<FDateTime&>(Expiration);
+
+			int32 NumAnnotations = MessageContext->GetAnnotations().Num();
+			Archive << NumAnnotations;
+
+			for (TMap<FName, FString>::TConstIterator It(MessageContext->GetAnnotations()); It; ++It)
+			{
+				Archive << const_cast<FName&>(It->Key);
+				Archive << const_cast<FString&>(It->Value);
+			}
+		}
+
+		// serialize message body
+		{
+			FJsonStructSerializerBackend Backend(Archive);
+			FStructSerializer::Serialize(MessageContext->GetMessage(), *MessageContext->GetMessageTypeInfo(), Backend);
+		}
+
+		SerializedMessage->UpdateState(EUdpSerializedMessageState::Complete);
+	}
+	else
+	{
+		SerializedMessage->UpdateState(EUdpSerializedMessageState::Invalid);
+	}
+}
+
+
+ENamedThreads::Type FUdpSerializeMessageTask::GetDesiredThread()
+{
+	return ENamedThreads::AnyThread;
+}
+
+
+TStatId FUdpSerializeMessageTask::GetStatId() const
+{
+	RETURN_QUICK_DECLARE_CYCLE_STAT(FUdpSerializeMessageTask, STATGROUP_TaskGraphTasks);
+}
+
+
+ESubsequentsMode::Type FUdpSerializeMessageTask::GetSubsequentsMode() 
+{ 
+	return ESubsequentsMode::TrackSubsequents; 
+}
diff --git a/Engine/Plugins/Messaging/UdpMessaging/Source/UdpMessaging/Private/Transport/UdpSerializeMessageTask.h b/Engine/Plugins/Messaging/UdpMessaging/Source/UdpMessaging/Private/Transport/UdpSerializeMessageTask.h
new file mode 100644
--- /dev/null
+++ b/Engine/Plugins/Messaging/UdpMessaging/Source/UdpMessaging/Private/Transport/UdpSerializeMessageTask.h
@@ -0,0 +1,65 @@
+// Copyright 1998-2014 Epic Games, Inc. All Rights Reserved.
+
+#pragma once
+
+#include "JsonStructSerializerBackend.h"
+#include "StructSerializer.h"
+
+
+/**
+ * Implements an asynchronous task for serializing a message.
+ */
+class FUdpSerializeMessageTask
+{
+public:
+
+	/**
+	 * Creates and initializes a new instance.
+	 *
+	 * @param InMessageContext The context of the message to serialize.
+	 * @param InSerializedMessage Will hold the serialized message data.
+	 */
+	FUdpSerializeMessageTask( IMessageContextRef InMessageContext, FUdpSerializedMessageRef InSerializedMessage )
+		: MessageContext(InMessageContext)
+		, SerializedMessage(InSerializedMessage)
+	{ }
+
+public:
+
+	/**
+	 * Performs the actual task.
+	 *
+	 * @param CurrentThread The thread that this task is executing on.
+	 * @param MyCompletionGraphEvent The completion event.
+	 */
+	void DoTask( ENamedThreads::Type CurrentThread, const FGraphEventRef& MyCompletionGraphEvent );
+	
+	/**
+	 * Returns the name of the thread that this task should run on.
+	 *
+	 * @return Thread name.
+	 */
+	ENamedThreads::Type GetDesiredThread();
+
+	/**
+	 * Gets the task's stats tracking identifier.
+	 *
+	 * @return Stats identifier.
+	 */
+	TStatId GetStatId() const;
+
+	/**
+	 * Gets the mode for tracking subsequent tasks.
+	 *
+	 * @return Always track subsequent tasks.
+	 */
+	static ESubsequentsMode::Type GetSubsequentsMode();
+
+private:
+
+	/** Holds the context of the message to serialize. */
+	IMessageContextRef MessageContext;
+
+	/** Holds a reference to the serialized message data. */
+	FUdpSerializedMessageRef SerializedMessage;
+};
diff --git a/Engine/Plugins/Messaging/UdpMessaging/Source/UdpMessaging/Private/Transport/UdpSerializedMessage.h b/Engine/Plugins/Messaging/UdpMessaging/Source/UdpMessaging/Private/Transport/UdpSerializedMessage.h
new file mode 100644
--- /dev/null
+++ b/Engine/Plugins/Messaging/UdpMessaging/Source/UdpMessaging/Private/Transport/UdpSerializedMessage.h
@@ -0,0 +1,104 @@
+// Copyright 1998-2014 Epic Games, Inc. All Rights Reserved.
+
+#pragma once
+
+
+/**
+ * Enumerates possibly states of a serialized message.
+ *
+ * @see FUdpSerializedMessage
+ */
+enum class EUdpSerializedMessageState
+{
+	/** The message data is complete. */
+	Complete,
+
+	/** The message data is incomplete. */
+	Incomplete,
+
+	/** The message data is invalid. */
+	Invalid
+};
+
+
+/**
+ * Holds serialized message data.
+ */
+class FUdpSerializedMessage
+	: public FMemoryWriter
+{
+public:
+
+	/** Default constructor. */
+	FUdpSerializedMessage()
+		: FMemoryWriter(Data, true)
+		, State(EUdpSerializedMessageState::Incomplete)
+	{ }
+
+public:
+
+	/**
+	 * Creates an archive reader to the data.
+	 *
+	 * The caller is responsible for deleting the returned object.
+	 *
+	 * @return An archive reader.
+	 */
+	FArchive* CreateReader()
+	{
+		return new FMemoryReader(Data, true);
+	}
+
+	/**
+	 * Gets the state of the message data.
+	 *
+	 * @return Message data state.
+	 * @see OnStateChanged
+	 */
+	EUdpSerializedMessageState GetState() const
+	{
+		return State;
+	}
+
+	/**
+	 * Returns a delegate that is executed when the message data's state changed.
+	 *
+	 * @return The delegate.
+	 * @see GetState
+	 */
+	FSimpleDelegate& OnStateChanged()
+	{
+		return StateChangedDelegate;
+	}
+
+	/**
+	 * Updates the state of this message data.
+	 *
+	 * @param InState The state to set.
+	 */
+	void UpdateState( EUdpSerializedMessageState InState )
+	{
+		State = InState;
+		StateChangedDelegate.ExecuteIfBound();
+	}
+
+private:
+
+	/** Holds the data. */
+	TArray<uint8> Data;
+
+	/** Holds the message data state. */
+	EUdpSerializedMessageState State;
+
+private:
+
+	/** Holds a delegate that is invoked when the data's state changed. */
+	FSimpleDelegate StateChangedDelegate;
+};
+
+
+/** Type definition for shared pointers to instances of FUdpSerializedMessage. */
+typedef TSharedPtr<FUdpSerializedMessage, ESPMode::ThreadSafe> FUdpSerializedMessagePtr;
+
+/** Type definition for shared references to instances of FUdpSerializedMessage. */
+typedef TSharedRef<FUdpSerializedMessage, ESPMode::ThreadSafe> FUdpSerializedMessageRef;
diff --git a/Engine/Plugins/Messaging/UdpMessaging/Source/UdpMessaging/Private/Tunnel/UdpMessageTunnel.cpp b/Engine/Plugins/Messaging/UdpMessaging/Source/UdpMessaging/Private/Tunnel/UdpMessageTunnel.cpp
new file mode 100644
--- /dev/null
+++ b/Engine/Plugins/Messaging/UdpMessaging/Source/UdpMessaging/Private/Tunnel/UdpMessageTunnel.cpp
@@ -0,0 +1,358 @@
+// Copyright 1998-2014 Epic Games, Inc. All Rights Reserved.
+
+#include "UdpMessagingPrivatePCH.h"
+
+
+/* FUdpMessageTunnel structors
+ *****************************************************************************/
+
+FUdpMessageTunnel::FUdpMessageTunnel( const FIPv4Endpoint& InUnicastEndpoint, const FIPv4Endpoint& InMulticastEndpoint )
+	: Listener(nullptr)
+	, MulticastEndpoint(InMulticastEndpoint)
+	, Stopping(false)
+	, TotalInboundBytes(0)
+	, TotalOutboundBytes(0)
+{
+	// initialize sockets
+	MulticastSocket = FUdpSocketBuilder(TEXT("UdpMessageMulticastSocket"))
+		.AsNonBlocking()
+		.AsReusable()
+		.BoundToAddress(InUnicastEndpoint.GetAddress())
+		.BoundToPort(InMulticastEndpoint.GetPort())
+		.JoinedToGroup(InMulticastEndpoint.GetAddress())
+		.WithMulticastLoopback()
+		.WithMulticastTtl(1);
+
+	UnicastSocket = FUdpSocketBuilder(TEXT("UdpMessageUnicastSocket"))
+		.AsNonBlocking()
+		.BoundToEndpoint(InUnicastEndpoint);
+
+	int32 NewSize = 0;
+	MulticastSocket->SetReceiveBufferSize(2 * 1024 * 1024, NewSize);
+	UnicastSocket->SetReceiveBufferSize(2 * 1024 * 1024, NewSize);
+
+	Thread = FRunnableThread::Create(this, TEXT("FUdpMessageTunnel"), 128 * 1024, TPri_AboveNormal);
+}
+
+
+FUdpMessageTunnel::~FUdpMessageTunnel()
+{
+	if (Thread != nullptr)
+	{
+		Thread->Kill(true);
+		delete Thread;
+	}
+
+	// destroy sockets
+	if (MulticastSocket != nullptr)
+	{
+		ISocketSubsystem::Get(PLATFORM_SOCKETSUBSYSTEM)->DestroySocket(MulticastSocket);
+		MulticastSocket = nullptr;
+	}
+	
+	if (UnicastSocket != nullptr)
+	{
+		ISocketSubsystem::Get(PLATFORM_SOCKETSUBSYSTEM)->DestroySocket(UnicastSocket);
+		UnicastSocket = nullptr;
+	}
+}
+
+
+/* FRunnable interface
+ *****************************************************************************/
+
+bool FUdpMessageTunnel::Init()
+{
+	return true;
+}
+
+
+uint32 FUdpMessageTunnel::Run()
+{
+	while (!Stopping)
+	{
+		CurrentTime = FDateTime::UtcNow();
+
+		UpdateConnections();
+
+		UdpToTcp(MulticastSocket);
+		UdpToTcp(UnicastSocket);
+
+		TcpToUdp();
+
+		RemoveExpiredNodes(LocalNodes);
+		RemoveExpiredNodes(RemoteNodes);
+	}
+
+	return 0;
+}
+
+
+void FUdpMessageTunnel::Stop()
+{
+	Stopping = true;
+}
+
+
+void FUdpMessageTunnel::Exit()
+{
+	// do nothing
+}
+
+
+/* IUdpMessageTunnel interface
+ *****************************************************************************/
+
+bool FUdpMessageTunnel::Connect( const FIPv4Endpoint& RemoteEndpoint )
+{
+	FSocket* Socket = FTcpSocketBuilder(TEXT("FUdpMessageTunnel.RemoteConnection"));
+
+	if (Socket != nullptr)
+	{
+		if (Socket->Connect(*RemoteEndpoint.ToInternetAddr()))
+		{
+			PendingConnections.Enqueue(MakeShareable(new FUdpMessageTunnelConnection(Socket, RemoteEndpoint)));
+
+			return true;
+		}
+		else
+		{
+			ISocketSubsystem::Get(PLATFORM_SOCKETSUBSYSTEM)->DestroySocket(Socket);
+		}
+	}
+
+	return false;
+}
+
+
+int32 FUdpMessageTunnel::GetConnections( TArray<IUdpMessageTunnelConnectionPtr>& OutConnections )
+{
+	FScopeLock Lock(&CriticalSection);
+
+	for (TArray<FUdpMessageTunnelConnectionPtr>::TConstIterator It(Connections); It; ++It)
+	{
+		OutConnections.Add(*It);
+	}
+
+	return OutConnections.Num();
+}
+
+
+uint64 FUdpMessageTunnel::GetTotalInboundBytes() const
+{
+	return TotalInboundBytes;
+}
+
+
+uint64 FUdpMessageTunnel::GetTotalOutboundBytes() const
+{
+	return TotalOutboundBytes;
+}
+
+
+bool FUdpMessageTunnel::IsServerRunning() const
+{
+	return (Listener != nullptr);
+}
+
+
+FSimpleDelegate& FUdpMessageTunnel::OnConnectionsChanged()
+{
+	return ConnectionsChangedDelegate;
+}
+
+
+void FUdpMessageTunnel::StartServer( const FIPv4Endpoint& LocalEndpoint )
+{
+	StopServer();
+
+	Listener = new FTcpListener(LocalEndpoint);
+	Listener->OnConnectionAccepted().BindRaw(this, &FUdpMessageTunnel::HandleListenerConnectionAccepted);
+}
+
+
+void FUdpMessageTunnel::StopServer()
+{
+	delete Listener;
+	Listener = nullptr;
+}
+
+
+/* FUdpMessageTunnel implementation
+ *****************************************************************************/
+
+void FUdpMessageTunnel::RemoveExpiredNodes( TMap<FGuid, FNodeInfo>& Nodes )
+{
+	for (TMap<FGuid, FNodeInfo>::TIterator It(Nodes); It; ++It)
+	{
+		if (CurrentTime - It.Value().LastDatagramReceivedTime > FTimespan::FromMinutes(2.0))
+		{
+			It.RemoveCurrent();
+		}
+	}
+}
+
+
+void FUdpMessageTunnel::TcpToUdp()
+{
+	TArray<FUdpMessageTunnelConnectionPtr>::TIterator It(Connections);
+	FArrayReaderPtr Payload;
+
+	while (It && UnicastSocket->Wait(ESocketWaitConditions::WaitForWrite, FTimespan::Zero()))
+	{
+		if ((*It)->Receive(Payload))
+		{
+			FUdpMessageSegment::FHeader Header;
+			*Payload << Header;
+
+			// check protocol version
+			if (Header.ProtocolVersion != UDP_MESSAGING_TRANSPORT_PROTOCOL_VERSION)
+			{
+				return;
+			}
+
+			Payload->Seek(0);
+
+			// update remote node & statistics
+			FNodeInfo& RemoteNode = RemoteNodes.FindOrAdd(Header.SenderNodeId);
+			RemoteNode.Connection = *It;
+			RemoteNode.LastDatagramReceivedTime = CurrentTime;
+
+			TotalInboundBytes += Payload->Num();
+
+			// forward payload to local nodes
+			FIPv4Endpoint RecipientEndpoint;
+
+			if (Header.RecipientNodeId.IsValid())
+			{
+				FNodeInfo* LocalNode = LocalNodes.Find(Header.RecipientNodeId);
+
+				if (LocalNode == nullptr)
+				{
+					continue;
+				}
+
+				RecipientEndpoint = LocalNode->Endpoint;
+			}
+			else
+			{
+				RecipientEndpoint = MulticastEndpoint;
+			}
+
+			int32 BytesSent = 0;
+			UnicastSocket->SendTo(Payload->GetData(), Payload->Num(), BytesSent, *RecipientEndpoint.ToInternetAddr());
+		}
+		else
+		{
+			++It;
+		}
+	}
+}
+
+
+void FUdpMessageTunnel::UdpToTcp( FSocket* Socket )
+{
+	TSharedRef<FInternetAddr> Sender = ISocketSubsystem::Get(PLATFORM_SOCKETSUBSYSTEM)->CreateInternetAddr();
+
+	uint32 DatagramSize = 0;
+
+	while (Socket->HasPendingData(DatagramSize))
+	{
+		FArrayReaderPtr Datagram = MakeShareable(new FArrayReader(true));
+		Datagram->Init(FMath::Min(DatagramSize, 65507u));
+
+		int32 BytesRead = 0;
+
+		if (Socket->RecvFrom(Datagram->GetData(), Datagram->Num(), BytesRead, *Sender))
+		{
+			FUdpMessageSegment::FHeader Header;
+			*Datagram << Header;
+
+			// check protocol version
+			if (Header.ProtocolVersion != UDP_MESSAGING_TRANSPORT_PROTOCOL_VERSION)
+			{
+				return;
+			}
+
+			// ignore loopback datagrams
+			if (RemoteNodes.Contains(Header.SenderNodeId))
+			{
+				return;
+			}
+
+			Datagram->Seek(0);
+
+			// forward datagram to remote nodes
+			if (Header.RecipientNodeId.IsValid())
+			{
+				FNodeInfo* RemoteNode = RemoteNodes.Find(Header.RecipientNodeId);
+
+				if ((RemoteNode != nullptr) && (RemoteNode->Connection.IsValid()))
+				{
+					RemoteNode->Connection->Send(Datagram);
+				}
+			}
+			else
+			{
+				for (int32 ConnectionIndex = 0; ConnectionIndex < Connections.Num(); ++ConnectionIndex)
+				{
+					Connections[ConnectionIndex]->Send(Datagram);
+				}
+			}
+
+			// update local node & statistics
+			FNodeInfo& LocalNode = LocalNodes.FindOrAdd(Header.SenderNodeId);
+			LocalNode.Endpoint = FIPv4Endpoint(Sender);
+			LocalNode.LastDatagramReceivedTime = CurrentTime;
+
+			TotalOutboundBytes += Datagram->Num();
+		}
+	}
+}
+
+
+void FUdpMessageTunnel::UpdateConnections()
+{
+	bool ConnectionsChanged = false;
+
+	// remove closed connections
+	for (int32 ConnectionIndex = Connections.Num() - 1; ConnectionIndex >= 0; --ConnectionIndex)
+	{
+		if (!Connections[ConnectionIndex]->IsOpen())
+		{
+			Connections.RemoveAtSwap(ConnectionIndex);
+
+			ConnectionsChanged = true;
+		}
+	}
+
+	// add pending connections
+	if (!PendingConnections.IsEmpty())
+	{
+		FUdpMessageTunnelConnectionPtr PendingConnection;
+
+		while (PendingConnections.Dequeue(PendingConnection))
+		{
+			Connections.Add(PendingConnection);
+		}
+
+		ConnectionsChanged = true;
+	}
+
+	// notify application
+	if (ConnectionsChanged)
+	{
+		ConnectionsChangedDelegate.ExecuteIfBound();
+	}
+}
+
+
+/* FUdpMessageTunnel callbacks
+ *****************************************************************************/
+
+bool FUdpMessageTunnel::HandleListenerConnectionAccepted( FSocket* ClientSocket, const FIPv4Endpoint& ClientEndpoint )
+{
+	PendingConnections.Enqueue(MakeShareable(new FUdpMessageTunnelConnection(ClientSocket, ClientEndpoint)));
+
+	return true;
+}
diff --git a/Engine/Plugins/Messaging/UdpMessaging/Source/UdpMessaging/Private/Tunnel/UdpMessageTunnel.h b/Engine/Plugins/Messaging/UdpMessaging/Source/UdpMessaging/Private/Tunnel/UdpMessageTunnel.h
--- a/Engine/Plugins/Messaging/UdpMessaging/Source/UdpMessaging/Private/Tunnel/UdpMessageTunnel.h
+++ b/Engine/Plugins/Messaging/UdpMessaging/Source/UdpMessaging/Private/Tunnel/UdpMessageTunnel.h
@@ -2,6 +2,13 @@
 
 #pragma once
 
+#include "IUdpMessageTunnel.h"
+
+
+// forward declarations
+class FRunnableThread;
+class FSocket;
+
 
 /**
  * Implements a bi-directional tunnel to send UDP messages over a TCP connection.
@@ -31,163 +38,32 @@ public:
 	 * @param InLocalEndpoint The local IP endpoint to receive messages on.
 	 * @param InMulticastEndpoint The multicast group endpoint to transport messages to.
 	 */
-	FUdpMessageTunnel( const FIPv4Endpoint& InUnicastEndpoint, const FIPv4Endpoint& InMulticastEndpoint )
-		: Listener(nullptr)
-		, MulticastEndpoint(InMulticastEndpoint)
-		, Stopping(false)
-		, TotalInboundBytes(0)
-		, TotalOutboundBytes(0)
-	{
-		// initialize sockets
-		MulticastSocket = FUdpSocketBuilder(TEXT("UdpMessageMulticastSocket"))
-			.AsNonBlocking()
-			.AsReusable()
-			.BoundToAddress(InUnicastEndpoint.GetAddress())
-			.BoundToPort(InMulticastEndpoint.GetPort())
-			.JoinedToGroup(InMulticastEndpoint.GetAddress())
-			.WithMulticastLoopback()
-			.WithMulticastTtl(1);
-
-		UnicastSocket = FUdpSocketBuilder(TEXT("UdpMessageUnicastSocket"))
-			.AsNonBlocking()
-			.BoundToEndpoint(InUnicastEndpoint);
-
-		int32 NewSize = 0;
-		MulticastSocket->SetReceiveBufferSize(2 * 1024 * 1024, NewSize);
-		UnicastSocket->SetReceiveBufferSize(2 * 1024 * 1024, NewSize);
-
-		Thread = FRunnableThread::Create(this, TEXT("FUdpMessageTunnel"), 128 * 1024, TPri_AboveNormal);
-	}
+	FUdpMessageTunnel( const FIPv4Endpoint& InUnicastEndpoint, const FIPv4Endpoint& InMulticastEndpoint );
 
 	/** Destructor. */
-	~FUdpMessageTunnel()
-	{
-		if (Thread != nullptr)
-		{
-			Thread->Kill(true);
-			delete Thread;
-		}
-
-		// destroy sockets
-		if (MulticastSocket != nullptr)
-		{
-			ISocketSubsystem::Get(PLATFORM_SOCKETSUBSYSTEM)->DestroySocket(MulticastSocket);
-			MulticastSocket = nullptr;
-		}
-	
-		if (UnicastSocket != nullptr)
-		{
-			ISocketSubsystem::Get(PLATFORM_SOCKETSUBSYSTEM)->DestroySocket(UnicastSocket);
-			UnicastSocket = nullptr;
-		}
-	}
+	~FUdpMessageTunnel();
 
 public:
 
 	// FRunnable interface
 
-	virtual bool Init() override
-	{
-		return true;
-	}
-
-	virtual uint32 Run() override
-	{
-		while (!Stopping)
-		{
-			CurrentTime = FDateTime::UtcNow();
-
-			UpdateConnections();
-
-			UdpToTcp(MulticastSocket);
-			UdpToTcp(UnicastSocket);
-
-			TcpToUdp();
-
-			RemoveExpiredNodes(LocalNodes);
-			RemoveExpiredNodes(RemoteNodes);
-		}
-
-		return 0;
-	}
-
-	virtual void Stop() override
-	{
-		Stopping = true;
-	}
-
-	virtual void Exit() override { }
+	virtual bool Init() override;
+	virtual uint32 Run() override;
+	virtual void Stop() override;
+	virtual void Exit() override;
 
 public:
 
 	// IUdpMessageTunnel interface
 
-	virtual bool Connect( const FIPv4Endpoint& RemoteEndpoint ) override
-	{
-		FSocket* Socket = FTcpSocketBuilder(TEXT("FUdpMessageTunnel.RemoteConnection"));
-
-		if (Socket != nullptr)
-		{
-			if (Socket->Connect(*RemoteEndpoint.ToInternetAddr()))
-			{
-				PendingConnections.Enqueue(MakeShareable(new FUdpMessageTunnelConnection(Socket, RemoteEndpoint)));
-
-				return true;
-			}
-			else
-			{
-				ISocketSubsystem::Get(PLATFORM_SOCKETSUBSYSTEM)->DestroySocket(Socket);
-			}
-		}
-
-		return false;
-	}
-
-	virtual int32 GetConnections( TArray<IUdpMessageTunnelConnectionPtr>& OutConnections ) override
-	{
-		FScopeLock Lock(&CriticalSection);
-
-		for (TArray<FUdpMessageTunnelConnectionPtr>::TConstIterator It(Connections); It; ++It)
-		{
-			OutConnections.Add(*It);
-		}
-
-		return OutConnections.Num();
-	}
-
-	virtual uint64 GetTotalInboundBytes() const override
-	{
-		return TotalInboundBytes;
-	}
-
-	virtual uint64 GetTotalOutboundBytes() const override
-	{
-		return TotalOutboundBytes;
-	}
-
-	virtual bool IsServerRunning() const override
-	{
-		return (Listener != nullptr);
-	}
-
-	virtual FSimpleDelegate& OnConnectionsChanged() override
-	{
-		return ConnectionsChangedDelegate;
-	}
-
-	virtual void StartServer( const FIPv4Endpoint& LocalEndpoint ) override
-	{
-		StopServer();
-
-		Listener = new FTcpListener(LocalEndpoint);
-		Listener->OnConnectionAccepted().BindRaw(this, &FUdpMessageTunnel::HandleListenerConnectionAccepted);
-	}
-
-	virtual void StopServer() override
-	{
-		delete Listener;
-		Listener = nullptr;
-	}
+	virtual bool Connect( const FIPv4Endpoint& RemoteEndpoint ) override;
+	virtual int32 GetConnections( TArray<IUdpMessageTunnelConnectionPtr>& OutConnections ) override;
+	virtual uint64 GetTotalInboundBytes() const override;
+	virtual uint64 GetTotalOutboundBytes() const override;
+	virtual bool IsServerRunning() const override;
+	virtual FSimpleDelegate& OnConnectionsChanged() override;
+	virtual void StartServer( const FIPv4Endpoint& LocalEndpoint ) override;
+	virtual void StopServer() override;
 
 protected:
 
@@ -196,184 +72,25 @@ protected:
 	 *
 	 * @param Nodes The collection of nodes to clean up.
 	 */
-	void RemoveExpiredNodes( TMap<FGuid, FNodeInfo>& Nodes )
-	{
-		for (TMap<FGuid, FNodeInfo>::TIterator It(Nodes); It; ++It)
-		{
-			if (CurrentTime - It.Value().LastDatagramReceivedTime > FTimespan::FromMinutes(2.0))
-			{
-				It.RemoveCurrent();
-			}
-		}
-	}
+	void RemoveExpiredNodes( TMap<FGuid, FNodeInfo>& Nodes );
 
 	/** Receives all pending payloads from the tunnels and forwards them to the local message bus. */
-	void TcpToUdp()
-	{
-		TArray<FUdpMessageTunnelConnectionPtr>::TIterator It(Connections);
-		FArrayReaderPtr Payload;
-
-		while (It && UnicastSocket->Wait(ESocketWaitConditions::WaitForWrite, FTimespan::Zero()))
-		{
-			if ((*It)->Receive(Payload))
-			{
-				FUdpMessageSegment::FHeader Header;
-				*Payload << Header;
-
-				// check protocol version
-				if (Header.ProtocolVersion != UDP_MESSAGING_TRANSPORT_PROTOCOL_VERSION)
-				{
-					return;
-				}
-
-				Payload->Seek(0);
-
-				// update remote node & statistics
-				FNodeInfo& RemoteNode = RemoteNodes.FindOrAdd(Header.SenderNodeId);
-				RemoteNode.Connection = *It;
-				RemoteNode.LastDatagramReceivedTime = CurrentTime;
-
-				TotalInboundBytes += Payload->Num();
-
-				// forward payload to local nodes
-				FIPv4Endpoint RecipientEndpoint;
-
-				if (Header.RecipientNodeId.IsValid())
-				{
-					FNodeInfo* LocalNode = LocalNodes.Find(Header.RecipientNodeId);
-
-					if (LocalNode == nullptr)
-					{
-						continue;
-					}
-
-					RecipientEndpoint = LocalNode->Endpoint;
-				}
-				else
-				{
-					RecipientEndpoint = MulticastEndpoint;
-				}
-
-				int32 BytesSent = 0;
-				UnicastSocket->SendTo(Payload->GetData(), Payload->Num(), BytesSent, *RecipientEndpoint.ToInternetAddr());
-			}
-			else
-			{
-				++It;
-			}
-		}
-	}
+	void TcpToUdp();
 
 	/**
 	 * Receives all buffered datagrams from the specified socket and forwards them to the tunnels.
 	 *
 	 * @param Socket The socket to receive from.
 	 */
-	void UdpToTcp( FSocket* Socket )
-	{
-		TSharedRef<FInternetAddr> Sender = ISocketSubsystem::Get(PLATFORM_SOCKETSUBSYSTEM)->CreateInternetAddr();
-
-		uint32 DatagramSize = 0;
-
-		while (Socket->HasPendingData(DatagramSize))
-		{
-			FArrayReaderPtr Datagram = MakeShareable(new FArrayReader(true));
-			Datagram->Init(FMath::Min(DatagramSize, 65507u));
-
-			int32 BytesRead = 0;
-
-			if (Socket->RecvFrom(Datagram->GetData(), Datagram->Num(), BytesRead, *Sender))
-			{
-				FUdpMessageSegment::FHeader Header;
-				*Datagram << Header;
-
-				// check protocol version
-				if (Header.ProtocolVersion != UDP_MESSAGING_TRANSPORT_PROTOCOL_VERSION)
-				{
-					return;
-				}
-
-				// ignore loopback datagrams
-				if (RemoteNodes.Contains(Header.SenderNodeId))
-				{
-					return;
-				}
-
-				Datagram->Seek(0);
-
-				// forward datagram to remote nodes
-				if (Header.RecipientNodeId.IsValid())
-				{
-					FNodeInfo* RemoteNode = RemoteNodes.Find(Header.RecipientNodeId);
-
-					if ((RemoteNode != nullptr) && (RemoteNode->Connection.IsValid()))
-					{
-						RemoteNode->Connection->Send(Datagram);
-					}
-				}
-				else
-				{
-					for (int32 ConnectionIndex = 0; ConnectionIndex < Connections.Num(); ++ConnectionIndex)
-					{
-						Connections[ConnectionIndex]->Send(Datagram);
-					}
-				}
-
-				// update local node & statistics
-				FNodeInfo& LocalNode = LocalNodes.FindOrAdd(Header.SenderNodeId);
-				LocalNode.Endpoint = FIPv4Endpoint(Sender);
-				LocalNode.LastDatagramReceivedTime = CurrentTime;
-
-				TotalOutboundBytes += Datagram->Num();
-			}
-		}
-	}
+	void UdpToTcp( FSocket* Socket );
 
 	/** Updates all active and pending connections. */
-	void UpdateConnections()
-	{
-		bool ConnectionsChanged = false;
-
-		// remove closed connections
-		for (int32 ConnectionIndex = Connections.Num() - 1; ConnectionIndex >= 0; --ConnectionIndex)
-		{
-			if (!Connections[ConnectionIndex]->IsOpen())
-			{
-				Connections.RemoveAtSwap(ConnectionIndex);
-
-				ConnectionsChanged = true;
-			}
-		}
-
-		// add pending connections
-		if (!PendingConnections.IsEmpty())
-		{
-			FUdpMessageTunnelConnectionPtr PendingConnection;
-
-			while (PendingConnections.Dequeue(PendingConnection))
-			{
-				Connections.Add(PendingConnection);
-			}
-
-			ConnectionsChanged = true;
-		}
-
-		// notify application
-		if (ConnectionsChanged)
-		{
-			ConnectionsChangedDelegate.ExecuteIfBound();
-		}
-	}
+	void UpdateConnections();
 
 private:
 
 	/** Callback for accepted connections to the local tunnel server. */
-	bool HandleListenerConnectionAccepted( FSocket* ClientSocket, const FIPv4Endpoint& ClientEndpoint )
-	{
-		PendingConnections.Enqueue(MakeShareable(new FUdpMessageTunnelConnection(ClientSocket, ClientEndpoint)));
-
-		return true;
-	}
+	bool HandleListenerConnectionAccepted( FSocket* ClientSocket, const FIPv4Endpoint& ClientEndpoint );
 
 private:
 
diff --git a/Engine/Plugins/Messaging/UdpMessaging/Source/UdpMessaging/Private/Tunnel/UdpMessageTunnelConnection.cpp b/Engine/Plugins/Messaging/UdpMessaging/Source/UdpMessaging/Private/Tunnel/UdpMessageTunnelConnection.cpp
new file mode 100644
--- /dev/null
+++ b/Engine/Plugins/Messaging/UdpMessaging/Source/UdpMessaging/Private/Tunnel/UdpMessageTunnelConnection.cpp
@@ -0,0 +1,238 @@
+// Copyright 1998-2014 Epic Games, Inc. All Rights Reserved.
+
+#include "UdpMessagingPrivatePCH.h"
+
+
+/* FUdpMessageTunnelConnection structors
+ *****************************************************************************/
+
+FUdpMessageTunnelConnection::FUdpMessageTunnelConnection( FSocket* InSocket, const FIPv4Endpoint& InRemoteEndpoint )
+	: OpenedTime(FDateTime::UtcNow())
+	, RemoteEndpoint(InRemoteEndpoint)
+	, Socket(InSocket)
+	, TotalBytesReceived(0)
+	, TotalBytesSent(0)
+	, bRun(false)
+{
+	int32 NewSize = 0;
+	Socket->SetReceiveBufferSize(2*1024*1024, NewSize);
+	Socket->SetSendBufferSize(2*1024*1024, NewSize);
+	Thread = FRunnableThread::Create(this, TEXT("FUdpMessageTunnelConnection"), 128 * 1024, TPri_Normal);
+}
+
+
+FUdpMessageTunnelConnection::~FUdpMessageTunnelConnection()
+{
+	if (Thread != nullptr)
+	{
+		Thread->Kill(true);
+		delete Thread;
+	}
+
+	delete Socket;
+	Socket = nullptr;
+}
+
+
+/* FUdpMessageTunnelConnection interface
+ *****************************************************************************/
+
+bool FUdpMessageTunnelConnection::Receive( FArrayReaderPtr& OutPayload )
+{
+	return Inbox.Dequeue(OutPayload);
+}
+
+
+bool FUdpMessageTunnelConnection::Send( const FArrayReaderPtr& Payload )
+{
+	if (IsOpen())
+	{
+		return Outbox.Enqueue(Payload);
+	}
+
+	return false;
+}
+
+
+/* FRunnable interface
+ *****************************************************************************/
+
+bool FUdpMessageTunnelConnection::Init()
+{
+	return (Socket != nullptr);
+}
+
+
+uint32 FUdpMessageTunnelConnection::Run()
+{
+	bRun = true;
+
+	while (bRun)
+	{
+		if (Socket->Wait(ESocketWaitConditions::WaitForReadOrWrite, FTimespan::Zero()))
+		{
+			bRun = ReceivePayloads();
+			bRun = SendPayloads();
+		}
+
+		FPlatformProcess::Sleep(0.0f);
+	}
+
+	ClosedTime = FDateTime::UtcNow();
+
+	return 0;
+}
+
+
+void FUdpMessageTunnelConnection::Stop()
+{
+	Socket->Close();
+	bRun = false;
+}
+
+
+void FUdpMessageTunnelConnection::Exit()
+{
+	// do nothing
+}
+
+
+/* IUdpMessageTunnelConnection interface
+ *****************************************************************************/
+
+void FUdpMessageTunnelConnection::Close()
+{
+	Stop();
+}
+
+
+uint64 FUdpMessageTunnelConnection::GetTotalBytesReceived() const
+{
+	return TotalBytesReceived;
+}
+
+
+uint64 FUdpMessageTunnelConnection::GetTotalBytesSent() const
+{
+	return TotalBytesSent;
+}
+
+
+FText FUdpMessageTunnelConnection::GetName() const
+{
+	return RemoteEndpoint.ToText();
+}
+
+
+FTimespan FUdpMessageTunnelConnection::GetUptime() const
+{
+	if (IsOpen())
+	{
+		return (FDateTime::UtcNow() - OpenedTime);
+	}
+		
+	return (ClosedTime - OpenedTime);
+}
+
+
+bool FUdpMessageTunnelConnection::IsOpen() const
+{
+	return (Socket->GetConnectionState() == SCS_Connected);
+}
+
+
+/* FUdpMessageTunnelConnection implementation
+ *****************************************************************************/
+
+bool FUdpMessageTunnelConnection::ReceivePayloads()
+{
+	uint32 PendingDataSize = 0;
+
+	// if we received a payload size...
+	while (Socket->HasPendingData(PendingDataSize) && (PendingDataSize >= sizeof(uint16)))
+	{
+		int32 BytesRead = 0;
+			
+		FArrayReader PayloadSizeData = FArrayReader(true);
+		PayloadSizeData.Init(sizeof(uint16));
+			
+		// ... read it from the stream without removing it...
+		if (!Socket->Recv(PayloadSizeData.GetData(), sizeof(uint16), BytesRead, ESocketReceiveFlags::Peek))
+		{
+			return false;
+		}
+
+		check(BytesRead == sizeof(uint16));
+
+		uint16 PayloadSize;
+		PayloadSizeData << PayloadSize;
+
+		// ... and if we received the complete payload...
+		if (Socket->HasPendingData(PendingDataSize) && (PendingDataSize >= sizeof(uint16) + PayloadSize))
+		{
+			// ... then remove the payload size from the stream...
+			if (!Socket->Recv((uint8*)&PayloadSize, sizeof(uint16), BytesRead))
+			{
+				return false;
+			}
+
+			check(BytesRead == sizeof(uint16));
+			TotalBytesReceived += BytesRead;
+
+			// ... receive the payload
+			FArrayReaderPtr Payload = MakeShareable(new FArrayReader(true));
+			Payload->Init(PayloadSize);
+
+			if (!Socket->Recv(Payload->GetData(), Payload->Num(), BytesRead))
+			{
+				return false;
+			}
+
+			check (BytesRead == PayloadSize);
+			TotalBytesReceived += BytesRead;
+
+			// .. and move it to the inbox
+			Inbox.Enqueue(Payload);
+
+			return true;
+		}
+	}
+
+	return true;
+}
+
+
+bool FUdpMessageTunnelConnection::SendPayloads()
+{
+	if (!Outbox.IsEmpty())
+	{
+		if (Socket->Wait(ESocketWaitConditions::WaitForWrite, FTimespan::Zero()))
+		{
+			FArrayReaderPtr Payload;
+			Outbox.Dequeue(Payload);
+			int32 BytesSent = 0;
+
+			// send the payload size
+			FArrayWriter PayloadSizeData = FArrayWriter(true);
+			uint16 PayloadSize = Payload->Num();
+			PayloadSizeData << PayloadSize;
+
+			if (!Socket->Send(PayloadSizeData.GetData(), sizeof(uint16), BytesSent))
+			{
+				return false;
+			}
+
+			TotalBytesSent += BytesSent;
+
+			// send the payload
+			if (!Socket->Send(Payload->GetData(), Payload->Num(), BytesSent))
+			{
+				return false;
+			}
+
+			TotalBytesSent += BytesSent;
+		}
+	}
+
+	return true;
+}
diff --git a/Engine/Plugins/Messaging/UdpMessaging/Source/UdpMessaging/Private/Tunnel/UdpMessageTunnelConnection.h b/Engine/Plugins/Messaging/UdpMessaging/Source/UdpMessaging/Private/Tunnel/UdpMessageTunnelConnection.h
--- a/Engine/Plugins/Messaging/UdpMessaging/Source/UdpMessaging/Private/Tunnel/UdpMessageTunnelConnection.h
+++ b/Engine/Plugins/Messaging/UdpMessaging/Source/UdpMessaging/Private/Tunnel/UdpMessageTunnelConnection.h
@@ -2,24 +2,12 @@
 
 #pragma once
 
-
-/** Type definition for shared pointers to instances of FUdpMessageTunnelConnection. */
-typedef TSharedPtr<class FUdpMessageTunnelConnection> FUdpMessageTunnelConnectionPtr;
-
-/** Type definition for shared references to instances of FUdpMessageTunnelConnection. */
-typedef TSharedRef<class FUdpMessageTunnelConnection> FUdpMessageTunnelConnectionRef;
-
-/** Type definition for thread-safe shared references to byte arrays. */
-typedef TSharedRef<TArray<uint8>, ESPMode::ThreadSafe> FSaveByteArrayRef;
+#include "IUdpMessageTunnelConnection.h"
 
 
-/**
- * Delegate type for received packets.
- *
- * The first parameter is the received packet data.
- * The second parameter is the sender's node identifier.
- */
-DECLARE_DELEGATE_TwoParams(FOnUdpMessageTunnelConnectionPacketReceived, FSaveByteArrayRef, const FGuid&)
+// forward declarations
+class FRunnableThread;
+class FSocket;
 
 
 /**
@@ -38,32 +26,10 @@ public:
 	 * @param InSocket The socket to use for this connection.
 	 * @param InRemoteAddress The IP endpoint of the remote client.
 	 */
-	FUdpMessageTunnelConnection( FSocket* InSocket, const FIPv4Endpoint& InRemoteEndpoint )
-		: OpenedTime(FDateTime::UtcNow())
-		, RemoteEndpoint(InRemoteEndpoint)
-		, Socket(InSocket)
-		, TotalBytesReceived(0)
-		, TotalBytesSent(0)
-		, bRun(false)
-	{
-		int32 NewSize = 0;
-		Socket->SetReceiveBufferSize(2*1024*1024, NewSize);
-		Socket->SetSendBufferSize(2*1024*1024, NewSize);
-		Thread = FRunnableThread::Create(this, TEXT("FUdpMessageTunnelConnection"), 128 * 1024, TPri_Normal);
-	}
+	FUdpMessageTunnelConnection( FSocket* InSocket, const FIPv4Endpoint& InRemoteEndpoint );
 
 	/** Destructor. */
-	~FUdpMessageTunnelConnection()
-	{
-		if (Thread != nullptr)
-		{
-			Thread->Kill(true);
-			delete Thread;
-		}
-
-		delete Socket;
-		Socket = nullptr;
-	}
+	~FUdpMessageTunnelConnection();
 
 public:
 
@@ -72,102 +38,37 @@ public:
 	 *
 	 * @param OutPayload Will hold the received payload, if any.
 	 * @return true if a payload was returned, false otherwise.
+	 * @see Send
 	 */
-	bool Receive( FArrayReaderPtr& OutPayload )
-	{
-		return Inbox.Dequeue(OutPayload);
-	}
+	bool Receive( FArrayReaderPtr& OutPayload );
 
 	/**
 	 * Sends a payload through this connection.
 	 *
 	 * @param Payload The payload to send.
+	 * @see Receive
 	 */
-	bool Send( const FArrayReaderPtr& Payload )
-	{
-		if (IsOpen())
-		{
-			return Outbox.Enqueue(Payload);
-		}
-
-		return false;
-	}
+	bool Send( const FArrayReaderPtr& Payload );
 
 public:
 
 	// FRunnable interface
 
-	virtual bool Init() override
-	{
-		return (Socket != nullptr);
-	}
-
-	virtual uint32 Run() override
-	{
-		bRun = true;
-
-		while (bRun)
-		{
-			if (Socket->Wait(ESocketWaitConditions::WaitForReadOrWrite, FTimespan::Zero()))
-			{
-				bRun = ReceivePayloads();
-				bRun = SendPayloads();
-			}
-
-			FPlatformProcess::Sleep(0.0f);
-		}
-
-		ClosedTime = FDateTime::UtcNow();
-
-		return 0;
-	}
-
-	virtual void Stop() override
-	{
-		Socket->Close();
-		bRun = false;
-	}
-
-	virtual void Exit() override { }
+	virtual bool Init() override;
+	virtual uint32 Run() override;
+	virtual void Stop() override;
+	virtual void Exit() override;
 
 public:
 
 	// IUdpMessageTunnelConnection interface
 
-	virtual void Close() override
-	{
-		Stop();
-	}
-
-	virtual uint64 GetTotalBytesReceived() const override
-	{
-		return TotalBytesReceived;
-	}
-
-	virtual uint64 GetTotalBytesSent() const override
-	{
-		return TotalBytesSent;
-	}
-
-	virtual FText GetName() const
-	{
-		return RemoteEndpoint.ToText();
-	}
-
-	virtual FTimespan GetUptime() const override
-	{
-		if (IsOpen())
-		{
-			return (FDateTime::UtcNow() - OpenedTime);
-		}
-		
-		return (ClosedTime - OpenedTime);
-	}
-
-	virtual bool IsOpen() const override
-	{
-		return (Socket->GetConnectionState() == SCS_Connected);
-	}
+	virtual void Close() override;
+	virtual uint64 GetTotalBytesReceived() const override;
+	virtual uint64 GetTotalBytesSent() const override;
+	virtual FText GetName() const;
+	virtual FTimespan GetUptime() const override;
+	virtual bool IsOpen() const override;
 
 protected:
 
@@ -175,102 +76,17 @@ protected:
 	 * Receives all pending payloads from the socket.
 	 *
 	 * @return true on success, false otherwise.
+	 * @see SendPayloads
 	 */
-	bool ReceivePayloads()
-	{
-		uint32 PendingDataSize = 0;
-
-		// if we received a payload size...
-		while (Socket->HasPendingData(PendingDataSize) && (PendingDataSize >= sizeof(uint16)))
-		{
-			int32 BytesRead = 0;
-			
-			FArrayReader PayloadSizeData = FArrayReader(true);
-			PayloadSizeData.Init(sizeof(uint16));
-			
-			// ... read it from the stream without removing it...
-			if (!Socket->Recv(PayloadSizeData.GetData(), sizeof(uint16), BytesRead, ESocketReceiveFlags::Peek))
-			{
-				return false;
-			}
-			check(BytesRead == sizeof(uint16));
-
-			uint16 PayloadSize;
-			PayloadSizeData << PayloadSize;
-
-			// ... and if we received the complete payload...
-			if (Socket->HasPendingData(PendingDataSize) && (PendingDataSize >= sizeof(uint16) + PayloadSize))
-			{
-				// ... then remove the payload size from the stream...
-				if (!Socket->Recv((uint8*)&PayloadSize, sizeof(uint16), BytesRead))
-				{
-					return false;
-				}
-				check(BytesRead == sizeof(uint16));
-
-				TotalBytesReceived += BytesRead;
-
-				// ... receive the payload
-				FArrayReaderPtr Payload = MakeShareable(new FArrayReader(true));
-				Payload->Init(PayloadSize);
-
-				if (!Socket->Recv(Payload->GetData(), Payload->Num(), BytesRead))
-				{
-					return false;
-				}
-				check (BytesRead == PayloadSize);
-
-				TotalBytesReceived += BytesRead;
-
-				// .. and move it to the inbox
-				Inbox.Enqueue(Payload);
-
-				return true;
-			}
-		}
-
-		return true;
-	}
+	bool ReceivePayloads();
 
 	/**
 	 * Sends all pending payloads to the socket.
 	 *
 	 * @return true on success, false otherwise.
+	 * @see ReceivePayloads
 	 */
-	bool SendPayloads()
-	{
-		if (!Outbox.IsEmpty())
-		{
-			if (Socket->Wait(ESocketWaitConditions::WaitForWrite, FTimespan::Zero()))
-			{
-				FArrayReaderPtr Payload;
-				Outbox.Dequeue(Payload);
-				int32 BytesSent = 0;
-
-				// send the payload size
-				FArrayWriter PayloadSizeData = FArrayWriter(true);
-				uint16 PayloadSize = Payload->Num();
-				PayloadSizeData << PayloadSize;
-
-				if (!Socket->Send(PayloadSizeData.GetData(), sizeof(uint16), BytesSent))
-				{
-					return false;
-				}
-
-				TotalBytesSent += BytesSent;
-
-				// send the payload
-				if (!Socket->Send(Payload->GetData(), Payload->Num(), BytesSent))
-				{
-					return false;
-				}
-
-				TotalBytesSent += BytesSent;
-			}
-		}
-
-		return true;
-	}
+	bool SendPayloads();
 
 private:
 
@@ -304,3 +120,10 @@ private:
 	/** thread should continue running. */
 	bool bRun;
 };
+
+
+/** Type definition for shared pointers to instances of FUdpMessageTunnelConnection. */
+typedef TSharedPtr<FUdpMessageTunnelConnection> FUdpMessageTunnelConnectionPtr;
+
+/** Type definition for shared references to instances of FUdpMessageTunnelConnection. */
+typedef TSharedRef<FUdpMessageTunnelConnection> FUdpMessageTunnelConnectionRef;
diff --git a/Engine/Plugins/Messaging/UdpMessaging/Source/UdpMessaging/Private/UdpMessagingModule.cpp b/Engine/Plugins/Messaging/UdpMessaging/Source/UdpMessaging/Private/UdpMessagingModule.cpp
--- a/Engine/Plugins/Messaging/UdpMessaging/Source/UdpMessaging/Private/UdpMessagingModule.cpp
+++ b/Engine/Plugins/Messaging/UdpMessaging/Source/UdpMessaging/Private/UdpMessagingModule.cpp
@@ -1,6 +1,8 @@
 // Copyright 1998-2014 Epic Games, Inc. All Rights Reserved.
 
 #include "UdpMessagingPrivatePCH.h"
+#include "ModuleInterface.h"
+#include "ModuleManager.h"
 
 
 #define LOCTEXT_NAMESPACE "FUdpMessagingModule"
@@ -129,7 +131,6 @@ protected:
 		GLog->Logf(TEXT("UdpMessaging: Initializing bridge on interface %s to multicast group %s."), *UnicastEndpoint.ToText().ToString(), *MulticastEndpoint.ToText().ToString());
 
 		MessageBridge = FMessageBridgeBuilder()
-			.UsingJsonSerializer()
 			.UsingTransport(MakeShareable(new FUdpMessageTransport(UnicastEndpoint, MulticastEndpoint, Settings->MulticastTimeToLive)));
 	}
 
diff --git a/Engine/Plugins/Messaging/UdpMessaging/Source/UdpMessaging/Private/UdpMessagingPrivatePCH.h b/Engine/Plugins/Messaging/UdpMessaging/Source/UdpMessaging/Private/UdpMessagingPrivatePCH.h
--- a/Engine/Plugins/Messaging/UdpMessaging/Source/UdpMessaging/Private/UdpMessagingPrivatePCH.h
+++ b/Engine/Plugins/Messaging/UdpMessaging/Source/UdpMessaging/Private/UdpMessagingPrivatePCH.h
@@ -2,16 +2,14 @@
 
 #pragma once
 
-#include "UdpMessaging.h"
-
 
 /* Private dependencies
  *****************************************************************************/
 
 #include "Core.h"
-#include "Json.h"
+#include "CoreUObject.h"
 #include "Messaging.h"
-#include "ModuleManager.h"
+#include "Networking.h"
 #include "Settings.h"
 #include "Sockets.h"
 #include "SocketSubsystem.h"
@@ -23,8 +21,14 @@
 /** Defines the default IP endpoint for multicast traffic. */
 #define UDP_MESSAGING_DEFAULT_MULTICAST_ENDPOINT FIPv4Endpoint(FIPv4Address(230, 0, 0, 1), 6666)
 
+/** Defines the maximum number of annotations a message can have. */
+#define UDP_MESSAGING_MAX_ANNOTATIONS 128
+
+/** Defines the maximum number of recipients a message can have. */
+#define UDP_MESSAGING_MAX_RECIPIENTS 1024
+
 /** Defines the protocol version of the UDP message transport. */
-#define UDP_MESSAGING_TRANSPORT_PROTOCOL_VERSION 9
+#define UDP_MESSAGING_TRANSPORT_PROTOCOL_VERSION 10
 
 
 /* Private includes
@@ -36,10 +40,13 @@
 
 // transport
 #include "UdpMessageBeacon.h"
-#include "ReassembledUdpMessage.h"
+#include "UdpReassembledMessage.h"
 #include "UdpMessageResequencer.h"
+#include "UdpDeserializedMessage.h"
+#include "UdpSerializedMessage.h"
 #include "UdpMessageSegmenter.h"
 #include "UdpMessageProcessor.h"
+#include "UdpSerializeMessageTask.h"
 #include "UdpMessageTransport.h"
 
 // tunnel
diff --git a/Engine/Plugins/Messaging/UdpMessaging/Source/UdpMessaging/Public/IUdpMessageTunnel.h b/Engine/Plugins/Messaging/UdpMessaging/Source/UdpMessaging/Public/IUdpMessageTunnel.h
new file mode 100644
--- /dev/null
+++ b/Engine/Plugins/Messaging/UdpMessaging/Source/UdpMessaging/Public/IUdpMessageTunnel.h
@@ -0,0 +1,95 @@
+// Copyright 1998-2014 Epic Games, Inc. All Rights Reserved.
+
+#pragma once
+
+
+// forward declarations
+class FIPv4Endpoint;
+class IUdpMessageTunnelConnection;
+
+
+/**
+ * Interface for message tunnels.
+ */
+class IUdpMessageTunnel
+{
+public:
+
+	/**
+	 * Establishes a tunnel with a remote server.
+	 *
+	 * @param RemoteEndpoint The endpoint of the server to connect to.
+	 * @return true if the connection has been established, false otherwise.
+	 * @see GetConnections
+	 */
+	virtual bool Connect( const FIPv4Endpoint& RemoteEndpoint ) = 0;
+
+	/** 
+	 * Starts the tunnel server.
+	 *
+	 * @param LocalEndpoint The IP endpoint to listen for incoming connections on.
+	 * @see IsServerRunning, StopServer
+	 */
+	virtual void StartServer( const FIPv4Endpoint& LocalEndpoint ) = 0;
+
+	/**
+	 * Stops the tunnel server.
+	 *
+	 * @see IsServerRunning, StartServer
+	 */
+	virtual void StopServer() = 0;
+
+public:
+
+	/**
+	 * Gets the list of all open tunnel connections.
+	 *
+	 * @param OutConnections Will hold the list of connections.
+	 * @return The number of connections returned.
+	 * @see Connect
+	 */
+	virtual int32 GetConnections( TArray<TSharedPtr<IUdpMessageTunnelConnection>>& OutConnections ) = 0;
+
+	/**
+	 * Gets the total number of bytes that were received from tunnels.
+	 *
+	 * @return Number of bytes.
+	 */
+	virtual uint64 GetTotalInboundBytes() const = 0;
+
+	/**
+	 * Gets the total number of bytes that were sent out through tunnels.
+	 *
+	 * @return Number of bytes.
+	 */
+	virtual uint64 GetTotalOutboundBytes() const = 0;
+
+	/**
+	 * Checks whether the tunnel server is running.
+	 *
+	 * @return true if the tunnel server is running, false otherwise.
+	 * @see StartServer, StopServer
+	 */
+	virtual bool IsServerRunning() const = 0;
+
+public:
+
+	/**
+	 * Gets a delegate that is executed when the list of incoming connections changed.
+	 *
+	 * @return The delegate.
+	 */
+	virtual FSimpleDelegate& OnConnectionsChanged() = 0;
+
+public:
+
+	/** Virtual destructor. */
+	virtual ~IUdpMessageTunnel() { }
+};
+
+
+/** Type definition for shared pointers to instances of IUdpMessageTunnel. */
+typedef TSharedPtr<IUdpMessageTunnel> IUdpMessageTunnelPtr;
+
+/** Type definition for shared references to instances of IUdpMessageTunnel. */
+typedef TSharedRef<IUdpMessageTunnel> IUdpMessageTunnelRef;
diff --git a/Engine/Plugins/Messaging/UdpMessaging/Source/UdpMessaging/Public/Interfaces/IUdpMessageTunnel.h b/Engine/Plugins/Messaging/UdpMessaging/Source/UdpMessaging/Public/Interfaces/IUdpMessageTunnel.h
deleted file mode 100644
diff --git a/Engine/Plugins/Messaging/UdpMessaging/Source/UdpMessaging/Public/UdpMessaging.h b/Engine/Plugins/Messaging/UdpMessaging/Source/UdpMessaging/Public/UdpMessaging.h
deleted file mode 100644
diff --git a/Engine/Plugins/Messaging/UdpMessaging/Source/UdpMessaging/UdpMessaging.Build.cs b/Engine/Plugins/Messaging/UdpMessaging/Source/UdpMessaging/UdpMessaging.Build.cs
--- a/Engine/Plugins/Messaging/UdpMessaging/Source/UdpMessaging/UdpMessaging.Build.cs
+++ b/Engine/Plugins/Messaging/UdpMessaging/Source/UdpMessaging/UdpMessaging.Build.cs
@@ -9,15 +9,16 @@ namespace UnrealBuildTool.Rules
 			PublicDependencyModuleNames.AddRange(
 				new string[]
 				{
-					"Networking",
+					"Core",
 				}
 			); 
 
 			PrivateDependencyModuleNames.AddRange(
 				new string[] {
-					"Core",
 					"CoreUObject",
 					"Json",
+					"Networking",
+					"Serialization",
 					"Sockets",
 				}
 			);
diff --git a/Engine/Source/Runtime/Messaging/Messaging.Build.cs b/Engine/Source/Runtime/Messaging/Messaging.Build.cs
--- a/Engine/Source/Runtime/Messaging/Messaging.Build.cs
+++ b/Engine/Source/Runtime/Messaging/Messaging.Build.cs
@@ -9,17 +9,10 @@ namespace UnrealBuildTool.Rules
 			PublicDependencyModuleNames.AddRange(
 				new string[]
 				{
+					"Core",
 					"CoreUObject",
 				}
 			); 
-			
-			PrivateDependencyModuleNames.AddRange(
-				new string[]
-				{
-					"Core",
-					"Json",
-				}
-			);
 
 			PrivateIncludePaths.AddRange(
 				new string[] {
diff --git a/Engine/Source/Runtime/Messaging/Private/Bridge/MessageAddressBook.h b/Engine/Source/Runtime/Messaging/Private/Bridge/MessageAddressBook.h
--- a/Engine/Source/Runtime/Messaging/Private/Bridge/MessageAddressBook.h
+++ b/Engine/Source/Runtime/Messaging/Private/Bridge/MessageAddressBook.h
@@ -4,7 +4,7 @@
 
 
 /**
- * Implements an address book that maps message addresses to transport nodes.
+ * Implements an address book that maps message addresses to remote nodes.
  */
 class FMessageAddressBook
 {
@@ -28,7 +28,7 @@ public:
 	 * Adds an address to the address book.
 	 *
 	 * @Address The message address to add.
-	 * @NodeId The identifier of the transport node that handles the message address.
+	 * @NodeId The identifier of the remote node that handles the message address.
 	 */
 	void Add( const FMessageAddress& Address, const FGuid& NodeId )
 	{
@@ -59,7 +59,7 @@ public:
 	}
 
 	/**
-	 * Gets the transport node identifiers for the specified list of message addresses.
+	 * Gets the remote node identifiers for the specified list of message addresses.
 	 *
 	 * @param Addresses The address list to retrieve the node identifiers for.
 	 * @return The list of node identifiers.
@@ -86,13 +86,13 @@ public:
 	/**
 	 * Removes all known message addresses.
 	 *
-	 * To remove only the addresses for a specific transport node, use RemoveNode().
-	 * If you are not interested in the removed addresses, use Clear().
+	 * To remove only the addresses for a specific remote node, use RemoveNode().
+	 * If you are not interested in the removed addresses, use Clear() instead.
 	 *
 	 * @param OutRemovedRecipients Will hold a list of recipients that were removed.
 	 * @see Clear, RemoveNode
 	 */
-	void RemoveAllNodes( TArray<FMessageAddress>& OutRemovedAddresses )
+	void RemoveAll( TArray<FMessageAddress>& OutRemovedAddresses )
 	{
 		OutRemovedAddresses.Reset();
 
@@ -103,9 +103,9 @@ public:
 	}
 
 	/**
-	 * Removes all known message addresses for the specified transport node identifier.
+	 * Removes all known message addresses for the specified remote node identifier.
 	 *
-	 * @param NodeId The identifier of the transport node to remove.
+	 * @param NodeId The identifier of the remote node to remove.
 	 * @param OutRemovedRecipients Will hold a list of recipients that were removed.
 	 * @see Clear, RemoveAllNodes
 	 */
@@ -134,6 +134,6 @@ private:
 	/** Holds a critical section to serialize access to the address book entries. */
 	FCriticalSection* CriticalSection;
 
-	/** Holds the collection of known addresses and their transport node identifiers. */
+	/** Holds the collection of known addresses and their remote node identifiers. */
 	TMap<FMessageAddress, FGuid> Entries;
 };
diff --git a/Engine/Source/Runtime/Messaging/Private/Bridge/MessageBridge.cpp b/Engine/Source/Runtime/Messaging/Private/Bridge/MessageBridge.cpp
--- a/Engine/Source/Runtime/Messaging/Private/Bridge/MessageBridge.cpp
+++ b/Engine/Source/Runtime/Messaging/Private/Bridge/MessageBridge.cpp
@@ -6,12 +6,11 @@
 /* FMessageBridge structors
  *****************************************************************************/
 
-FMessageBridge::FMessageBridge( const FMessageAddress InAddress, const IMessageBusRef& InBus, const ISerializeMessagesRef& InSerializer, const IMessageTransportRef& InTransport )
+FMessageBridge::FMessageBridge( const FMessageAddress InAddress, const IMessageBusRef& InBus, const IMessageTransportRef& InTransport )
 	: Address(InAddress)
 	, Bus(InBus)
 	, Enabled(false)
 	, Id(FGuid::NewGuid())
-	, Serializer(InSerializer)
 	, Transport(InTransport)
 {
 	Bus->OnShutdown().AddRaw(this, &FMessageBridge::HandleMessageBusShutdown);
@@ -20,13 +19,14 @@ FMessageBridge::FMessageBridge( const FMessageAddress InAddress, const IMessageB
 }
 
 
-FMessageBridge::~FMessageBridge( )
+FMessageBridge::~FMessageBridge()
 {
 	Shutdown();
 
 	if (Bus.IsValid())
 	{
 		Bus->OnShutdown().RemoveAll(this);
+		Bus->Unregister(Address);
 
 		TArray<FMessageAddress> RemovedAddresses;
 
@@ -41,7 +41,7 @@ FMessageBridge::~FMessageBridge( )
 /* IMessageBridge interface
  *****************************************************************************/
 
-void FMessageBridge::Disable( ) 
+void FMessageBridge::Disable() 
 {
 	if (!Enabled)
 	{
@@ -63,7 +63,7 @@ void FMessageBridge::Disable( )
 }
 
 
-void FMessageBridge::Enable( )
+void FMessageBridge::Enable()
 {
 	if (Enabled || !Bus.IsValid() || !Transport.IsValid())
 	{
@@ -76,6 +76,8 @@ void FMessageBridge::Enable( )
 		return;
 	}
 
+	Bus->Register(Address, AsShared());
+
 	if (MessageSubscription.IsValid())
 	{
 		MessageSubscription->Enable();
@@ -99,22 +101,21 @@ void FMessageBridge::ReceiveMessage( const IMessageContextRef& Context )
 		return;
 	}
 
-	TArray<FGuid> TransportNodes;
+	// get remote nodes
+	TArray<FGuid> RemoteNodes;
 
 	if (Context->GetRecipients().Num() > 0)
 	{
-		TransportNodes = AddressBook.GetNodesFor(Context->GetRecipients());
+		RemoteNodes = AddressBook.GetNodesFor(Context->GetRecipients());
 
-		if (TransportNodes.Num() == 0)
+		if (RemoteNodes.Num() == 0)
 		{
 			return;
 		}
 	}
 
-	// forward message to transport nodes
-	FMessageDataRef MessageData = MakeShareable(new FMessageData());
-	TGraphTask<FMessageSerializeTask>::CreateTask().ConstructAndDispatchWhenReady(Context, MessageData, Serializer.ToSharedRef());
-	Transport->TransportMessage(MessageData, Context->GetAttachment(), TransportNodes);
+	// forward message to remote nodes
+	Transport->TransportMessage(Context, RemoteNodes);
 }
 
 
@@ -130,7 +131,7 @@ void FMessageBridge::NotifyMessageError( const IMessageContextRef& Context, cons
 /* FMessageBridge implementation
  *****************************************************************************/
 
-void FMessageBridge::Shutdown( )
+void FMessageBridge::Shutdown()
 {
 	Disable();
 
@@ -145,40 +146,35 @@ void FMessageBridge::Shutdown( )
 /* FMessageBridge callbacks
  *****************************************************************************/
 
-void FMessageBridge::HandleMessageBusShutdown( )
+void FMessageBridge::HandleMessageBusShutdown()
 {
 	Shutdown();
 	Bus.Reset();
 }
 
 
-void FMessageBridge::HandleTransportMessageReceived( FArchive& MessageData, const IMessageAttachmentPtr& Attachment, const FGuid& NodeId )
+void FMessageBridge::HandleTransportMessageReceived( const IMessageContextRef& Context, const FGuid& NodeId )
 {
 	if (!Enabled || !Bus.IsValid())
 	{
 		return;
 	}
 
-	IMutableMessageContextRef MessageContext = MakeShareable(new FMessageContext());
-
-	if (Serializer->DeserializeMessage(MessageData, MessageContext))
+	// discard expired messages
+	if (Context->GetExpiration() < FDateTime::UtcNow())
 	{
-		if (MessageContext->GetExpiration() >= FDateTime::UtcNow())
-		{
-			// register newly discovered endpoints
-			if (!AddressBook.Contains(MessageContext->GetSender()))
-			{
-				AddressBook.Add(MessageContext->GetSender(), NodeId);
-
-				Bus->Register(MessageContext->GetSender(), AsShared());
-			}
-
-			// forward the message to the internal bus
-			MessageContext->SetAttachment(Attachment);
+		return;
+	}
 
-			Bus->Forward(MessageContext, MessageContext->GetRecipients(), EMessageScope::Process, FTimespan::Zero(), AsShared());
-		}
+	// register newly discovered endpoints
+	if (!AddressBook.Contains(Context->GetSender()))
+	{
+		AddressBook.Add(Context->GetSender(), NodeId);
+		Bus->Register(Context->GetSender(), AsShared());
 	}
+
+	// forward message to local bus
+	Bus->Forward(Context, Context->GetRecipients(), FTimespan::Zero(), AsShared());
 }
 
 
diff --git a/Engine/Source/Runtime/Messaging/Private/Bridge/MessageBridge.h b/Engine/Source/Runtime/Messaging/Private/Bridge/MessageBridge.h
--- a/Engine/Source/Runtime/Messaging/Private/Bridge/MessageBridge.h
+++ b/Engine/Source/Runtime/Messaging/Private/Bridge/MessageBridge.h
@@ -2,9 +2,26 @@
 
 #pragma once
 
+#include "IMessageBridge.h"
+#include "IMessageContext.h"
+#include "IMessageTransport.h"
+#include "IReceiveMessages.h"
+#include "ISendMessages.h"
+
 
 /**
  * Implements a message bridge.
+ *
+ * A message bridge is a special message endpoint that connects multiple message buses
+ * running in different processes or on different devices. This allows messages that are
+ * available in one system to also be available on other systems.
+ *
+ * Message bridges use an underlying transport layer to channel the messages between two
+ * or more systems. Such layers may utilize system specific technologies, such as network
+ * sockets or shared memory to communicate with remote bridges. The bridge acts as a map
+ * from message addresses to remote nodes and vice versa.
+ *
+ * @see IMessageBus, IMessageTransport
  */
 class FMessageBridge
 	: public TSharedFromThis<FMessageBridge, ESPMode::ThreadSafe>
@@ -17,12 +34,11 @@ public:
 	/**
 	 * Creates and initializes a new instance.
 	 *
-	 * @param InAddress The address for this bridge.
+	 * @param InAddress The message address for this bridge.
 	 * @param InBus The message bus that this node is connected to.
-	 * @param InSerializer The message serializer to use.
 	 * @param InTransport The transport mechanism to use.
 	 */
-	FMessageBridge( const FMessageAddress InAddress, const IMessageBusRef& InBus, const ISerializeMessagesRef& InSerializer, const IMessageTransportRef& InTransport );
+	FMessageBridge( const FMessageAddress InAddress, const IMessageBusRef& InBus, const IMessageTransportRef& InTransport );
 
 	/** Destructor. */
 	~FMessageBridge();
@@ -83,13 +99,13 @@ protected:
 
 private:
 
-	/** Callback for message bus shut downs. */
+	/** Callback for message bus shutdowns. */
 	void HandleMessageBusShutdown();
 
-	/** Callback for received transport messages. */
-	void HandleTransportMessageReceived( FArchive& MessageData, const IMessageAttachmentPtr& Attachment, const FGuid& NodeId );
+	/** Callback for messages received from the transport layer. */
+	void HandleTransportMessageReceived( const IMessageContextRef& Envelope, const FGuid& NodeId );
 
-	/** Callback for lost transport endpoints. */
+	/** Callback for lost remote nodes. */
 	void HandleTransportNodeLost( const FGuid& LostNodeId );
 
 private:
@@ -112,9 +128,6 @@ private:
 	/** Holds the message subscription for outbound messages. */
 	IMessageSubscriptionPtr MessageSubscription;
 
-	/** Holds the message serializer. */
-	ISerializeMessagesPtr Serializer;
-
 	/** Holds the message transport object. */
 	IMessageTransportPtr Transport;
 };
diff --git a/Engine/Source/Runtime/Messaging/Private/Bridge/MessageData.h b/Engine/Source/Runtime/Messaging/Private/Bridge/MessageData.h
deleted file mode 100644
diff --git a/Engine/Source/Runtime/Messaging/Private/Bridge/MessageDeserializeTask.h b/Engine/Source/Runtime/Messaging/Private/Bridge/MessageDeserializeTask.h
deleted file mode 100644
diff --git a/Engine/Source/Runtime/Messaging/Private/Bridge/MessageForwardTask.h b/Engine/Source/Runtime/Messaging/Private/Bridge/MessageForwardTask.h
deleted file mode 100644
diff --git a/Engine/Source/Runtime/Messaging/Private/Bridge/MessageSerializeTask.h b/Engine/Source/Runtime/Messaging/Private/Bridge/MessageSerializeTask.h
deleted file mode 100644
diff --git a/Engine/Source/Runtime/Messaging/Private/Bus/MessageBus.cpp b/Engine/Source/Runtime/Messaging/Private/Bus/MessageBus.cpp
--- a/Engine/Source/Runtime/Messaging/Private/Bus/MessageBus.cpp
+++ b/Engine/Source/Runtime/Messaging/Private/Bus/MessageBus.cpp
@@ -27,6 +27,54 @@ FMessageBus::~FMessageBus( )
 /* IMessageBus interface
  *****************************************************************************/
 
+void FMessageBus::Forward( const IMessageContextRef& Context, const TArray<FMessageAddress>& Recipients, const FTimespan& Delay, const ISendMessagesRef& Forwarder )
+{
+	Router->RouteMessage(MakeShareable(new FMessageContext(Context, Forwarder->GetSenderAddress(), Recipients, EMessageScope::Process, FDateTime::UtcNow() + Delay, FTaskGraphInterface::Get().GetCurrentThreadIfKnown())));
+}
+
+
+IMessageTracerRef FMessageBus::GetTracer()
+{
+	return Router->GetTracer();
+}
+
+
+void FMessageBus::Intercept( const IMessageInterceptorRef& Interceptor, const FName& MessageType )
+{
+	if (MessageType != NAME_None)
+	{
+		if (!RecipientAuthorizer.IsValid() || RecipientAuthorizer->AuthorizeInterceptor(Interceptor, MessageType))
+		{
+			Router->AddInterceptor(Interceptor, MessageType);
+		}			
+	}
+}
+
+
+FOnMessageBusShutdown& FMessageBus::OnShutdown()
+{
+	return ShutdownDelegate;
+}
+
+
+void FMessageBus::Publish( void* Message, UScriptStruct* TypeInfo, EMessageScope Scope, const FTimespan& Delay, const FDateTime& Expiration, const ISendMessagesRef& Publisher )
+{
+	Router->RouteMessage(MakeShareable(new FMessageContext(Message, TypeInfo, nullptr, Publisher->GetSenderAddress(), TArray<FMessageAddress>(), Scope, FDateTime::UtcNow() + Delay, Expiration, FTaskGraphInterface::Get().GetCurrentThreadIfKnown())));
+}
+
+
+void FMessageBus::Register( const FMessageAddress& Address, const IReceiveMessagesRef& Recipient )
+{
+	Router->AddRecipient(Address, Recipient);
+}
+
+
+void FMessageBus::Send( void* Message, UScriptStruct* TypeInfo, const IMessageAttachmentPtr& Attachment, const TArray<FMessageAddress>& Recipients, const FTimespan& Delay, const FDateTime& Expiration, const ISendMessagesRef& Sender )
+{
+	Router->RouteMessage(MakeShareable(new FMessageContext(Message, TypeInfo, Attachment, Sender->GetSenderAddress(), Recipients, EMessageScope::Network, FDateTime::UtcNow() + Delay, Expiration, FTaskGraphInterface::Get().GetCurrentThreadIfKnown())));
+}
+
+
 void FMessageBus::Shutdown( )
 {
 	if (RouterThread != nullptr)
@@ -38,3 +86,50 @@ void FMessageBus::Shutdown( )
 		RouterThread = nullptr;
 	}
 }
+
+
+IMessageSubscriptionPtr FMessageBus::Subscribe( const IReceiveMessagesRef& Subscriber, const FName& MessageType, const FMessageScopeRange& ScopeRange )
+{
+	if (MessageType != NAME_None)
+	{
+		if (!RecipientAuthorizer.IsValid() || RecipientAuthorizer->AuthorizeSubscription(Subscriber, MessageType))
+		{
+			IMessageSubscriptionRef Subscription = MakeShareable(new FMessageSubscription(Subscriber, MessageType, ScopeRange));
+			Router->AddSubscription(Subscription);
+
+			return Subscription;
+		}
+	}
+
+	return nullptr;
+}
+
+
+void FMessageBus::Unintercept( const IMessageInterceptorRef& Interceptor, const FName& MessageType )
+{
+	if (MessageType != NAME_None)
+	{
+		Router->RemoveInterceptor(Interceptor, MessageType);
+	}
+}
+
+
+void FMessageBus::Unregister( const FMessageAddress& Address )
+{
+	if (!RecipientAuthorizer.IsValid() || RecipientAuthorizer->AuthorizeUnregistration(Address))
+	{
+		Router->RemoveRecipient(Address);
+	}
+}
+
+
+void FMessageBus::Unsubscribe( const IReceiveMessagesRef& Subscriber, const FName& MessageType )
+{
+	if (MessageType != NAME_None)
+	{
+		if (!RecipientAuthorizer.IsValid() || RecipientAuthorizer->AuthorizeUnsubscription(Subscriber, MessageType))
+		{
+			Router->RemoveSubscription(Subscriber, MessageType);
+		}
+	}
+}
diff --git a/Engine/Source/Runtime/Messaging/Private/Bus/MessageBus.h b/Engine/Source/Runtime/Messaging/Private/Bus/MessageBus.h
--- a/Engine/Source/Runtime/Messaging/Private/Bus/MessageBus.h
+++ b/Engine/Source/Runtime/Messaging/Private/Bus/MessageBus.h
@@ -2,6 +2,20 @@
 
 #pragma once
 
+#include "IAuthorizeMessageRecipients.h"
+#include "IMessageAttachment.h"
+#include "IMessageBus.h"
+#include "IMessageContext.h"
+#include "IMessageInterceptor.h"
+#include "IMessageSubscription.h"
+#include "IReceiveMessages.h"
+#include "ISendMessages.h"
+
+
+// forward declarations
+class FMessageRouter;
+class FRunnableThread;
+
 
 /**
  * Implements a message bus.
@@ -19,98 +33,25 @@ public:
 	 */
 	FMessageBus( const IAuthorizeMessageRecipientsPtr& InRecipientAuthorizer );
 
-	/** Destructor */
+	/** Destructor. */
 	~FMessageBus();
 
 public:
 
 	// IMessageBus interface
 
-	virtual void Forward( const IMessageContextRef& Context, const TArray<FMessageAddress>& Recipients, EMessageScope ForwardingScope, const FTimespan& Delay, const ISendMessagesRef& Forwarder ) override
-	{
-		Router->RouteMessage(MakeShareable(new FMessageContext(Context, Forwarder->GetSenderAddress(), Recipients, ForwardingScope, FDateTime::UtcNow(), FTaskGraphInterface::Get().GetCurrentThreadIfKnown())));
-	}
-
-	virtual IMessageTracerRef GetTracer() override
-	{
-		return Router->GetTracer();
-	}
-
-	virtual void Intercept( const IMessageInterceptorRef& Interceptor, const FName& MessageType ) override
-	{
-		if (MessageType != NAME_None)
-		{
-			if (!RecipientAuthorizer.IsValid() || RecipientAuthorizer->AuthorizeInterceptor(Interceptor, MessageType))
-			{
-				Router->AddInterceptor(Interceptor, MessageType);
-			}			
-		}
-	}
-
-	virtual FOnMessageBusShutdown& OnShutdown() override
-	{
-		return ShutdownDelegate;
-	}
-
-	virtual void Publish( void* Message, UScriptStruct* TypeInfo, EMessageScope Scope, const FTimespan& Delay, const FDateTime& Expiration, const ISendMessagesRef& Publisher ) override
-	{
-		Router->RouteMessage(MakeShareable(new FMessageContext(Message, TypeInfo, nullptr, Publisher->GetSenderAddress(), TArray<FMessageAddress>(), Scope, FDateTime::UtcNow() + Delay, Expiration, FTaskGraphInterface::Get().GetCurrentThreadIfKnown())));
-	}
-
-	virtual void Register( const FMessageAddress& Address, const IReceiveMessagesRef& Recipient ) override
-	{
-		Router->AddRecipient(Address, Recipient);
-	}
-
-	virtual void Send( void* Message, UScriptStruct* TypeInfo, const IMessageAttachmentPtr& Attachment, const TArray<FMessageAddress>& Recipients, const FTimespan& Delay, const FDateTime& Expiration, const ISendMessagesRef& Sender ) override
-	{
-		Router->RouteMessage(MakeShareable(new FMessageContext(Message, TypeInfo, Attachment, Sender->GetSenderAddress(), Recipients, EMessageScope::Network, FDateTime::UtcNow() + Delay, Expiration, FTaskGraphInterface::Get().GetCurrentThreadIfKnown())));
-	}
-
+	virtual void Forward( const IMessageContextRef& Context, const TArray<FMessageAddress>& Recipients, const FTimespan& Delay, const ISendMessagesRef& Forwarder ) override;
+	virtual IMessageTracerRef GetTracer() override;
+	virtual void Intercept( const IMessageInterceptorRef& Interceptor, const FName& MessageType ) override;
+	virtual FOnMessageBusShutdown& OnShutdown() override;
+	virtual void Publish( void* Message, UScriptStruct* TypeInfo, EMessageScope Scope, const FTimespan& Delay, const FDateTime& Expiration, const ISendMessagesRef& Publisher ) override;
+	virtual void Register( const FMessageAddress& Address, const IReceiveMessagesRef& Recipient ) override;
+	virtual void Send( void* Message, UScriptStruct* TypeInfo, const IMessageAttachmentPtr& Attachment, const TArray<FMessageAddress>& Recipients, const FTimespan& Delay, const FDateTime& Expiration, const ISendMessagesRef& Sender ) override;
 	virtual void Shutdown() override;
-
-	virtual IMessageSubscriptionPtr Subscribe( const IReceiveMessagesRef& Subscriber, const FName& MessageType, const FMessageScopeRange& ScopeRange ) override
-	{
-		if (MessageType != NAME_None)
-		{
-			if (!RecipientAuthorizer.IsValid() || RecipientAuthorizer->AuthorizeSubscription(Subscriber, MessageType))
-			{
-				IMessageSubscriptionRef Subscription = MakeShareable(new FMessageSubscription(Subscriber, MessageType, ScopeRange));
-				Router->AddSubscription(Subscription);
-
-				return Subscription;
-			}
-		}
-
-		return nullptr;
-	}
-
-	virtual void Unintercept( const IMessageInterceptorRef& Interceptor, const FName& MessageType ) override
-	{
-		if (MessageType != NAME_None)
-		{
-			Router->RemoveInterceptor(Interceptor, MessageType);
-		}
-	}
-
-	virtual void Unregister( const FMessageAddress& Address ) override
-	{
-		if (!RecipientAuthorizer.IsValid() || RecipientAuthorizer->AuthorizeUnregistration(Address))
-		{
-			Router->RemoveRecipient(Address);
-		}
-	}
-
-	virtual void Unsubscribe( const IReceiveMessagesRef& Subscriber, const FName& MessageType ) override
-	{
-		if (MessageType != NAME_None)
-		{
-			if (!RecipientAuthorizer.IsValid() || RecipientAuthorizer->AuthorizeUnsubscription(Subscriber, MessageType))
-			{
-				Router->RemoveSubscription(Subscriber, MessageType);
-			}
-		}
-	}
+	virtual IMessageSubscriptionPtr Subscribe( const IReceiveMessagesRef& Subscriber, const FName& MessageType, const FMessageScopeRange& ScopeRange ) override;
+	virtual void Unintercept( const IMessageInterceptorRef& Interceptor, const FName& MessageType ) override;
+	virtual void Unregister( const FMessageAddress& Address ) override;
+	virtual void Unsubscribe( const IReceiveMessagesRef& Subscriber, const FName& MessageType ) override;
 
 private:
 
diff --git a/Engine/Source/Runtime/Messaging/Private/Bus/MessageContext.cpp b/Engine/Source/Runtime/Messaging/Private/Bus/MessageContext.cpp
new file mode 100644
--- /dev/null
+++ b/Engine/Source/Runtime/Messaging/Private/Bus/MessageContext.cpp
@@ -0,0 +1,129 @@
+// Copyright 1998-2014 Epic Games, Inc. All Rights Reserved.
+
+#include "MessagingPrivatePCH.h"
+
+
+/* FMessageContext structors
+ *****************************************************************************/
+
+FMessageContext::~FMessageContext()
+{
+	if (Message != nullptr)
+	{
+		if (TypeInfo.IsValid())
+		{
+			TypeInfo->DestroyScriptStruct(Message);
+		}
+
+		FMemory::Free(Message);
+	}		
+}
+
+
+/* IMessageContext interface
+ *****************************************************************************/
+
+const TMap<FName, FString>& FMessageContext::GetAnnotations() const
+{
+	if (OriginalContext.IsValid())
+	{
+		return OriginalContext->GetAnnotations();
+	}
+
+	return Annotations;
+}
+
+
+IMessageAttachmentPtr FMessageContext::GetAttachment() const
+{
+	if (OriginalContext.IsValid())
+	{
+		return OriginalContext->GetAttachment();
+	}
+
+	return Attachment;
+}
+
+
+const FDateTime& FMessageContext::GetExpiration() const
+{
+	if (OriginalContext.IsValid())
+	{
+		return OriginalContext->GetExpiration();
+	}
+
+	return Expiration;
+}
+
+
+const void* FMessageContext::GetMessage() const
+{
+	if (OriginalContext.IsValid())
+	{
+		return OriginalContext->GetMessage();
+	}
+
+	return Message;
+}
+
+
+const TWeakObjectPtr<UScriptStruct>& FMessageContext::GetMessageTypeInfo() const
+{
+	if (OriginalContext.IsValid())
+	{
+		return OriginalContext->GetMessageTypeInfo();
+	}
+
+	return TypeInfo;
+}
+
+
+IMessageContextPtr FMessageContext::GetOriginalContext() const
+{
+	return OriginalContext;
+}
+
+const TArray<FMessageAddress>& FMessageContext::GetRecipients() const
+{
+	return Recipients;
+}
+
+
+EMessageScope FMessageContext::GetScope() const
+{
+	return Scope;
+}
+
+
+const FMessageAddress& FMessageContext::GetSender() const
+{
+	if (OriginalContext.IsValid())
+	{
+		return OriginalContext->GetSender();
+	}
+
+	return Sender;
+}
+
+
+ENamedThreads::Type FMessageContext::GetSenderThread() const
+{
+	return SenderThread;
+}
+
+
+const FDateTime& FMessageContext::GetTimeForwarded() const
+{
+	return TimeSent;
+}
+
+
+const FDateTime& FMessageContext::GetTimeSent() const
+{
+	if (OriginalContext.IsValid())
+	{
+		return OriginalContext->GetTimeSent();
+	}
+
+	return TimeSent;
+}
diff --git a/Engine/Source/Runtime/Messaging/Private/Bus/MessageContext.h b/Engine/Source/Runtime/Messaging/Private/Bus/MessageContext.h
--- a/Engine/Source/Runtime/Messaging/Private/Bus/MessageContext.h
+++ b/Engine/Source/Runtime/Messaging/Private/Bus/MessageContext.h
@@ -2,15 +2,23 @@
 
 #pragma once
 
+#include "IMessageAttachment.h"
+#include "IMessageContext.h"
+
+
+// forward declarations
+class UScriptStruct;
+
 
 /**
- * Implements a message context.
+ * Implements a message context for messages sent through the message bus.
  *
- * The message context holds out-of-band information for messages that allows the
- * messaging system to route and process them correctly.
+ * Message contexts contain a message and additional data about that message,
+ * such as when the message was sent, who sent it and where it is being sent to.
  */
+
 class FMessageContext
-	: public IMutableMessageContext
+	: public IMessageContext
 {
 public:
 
@@ -23,6 +31,8 @@ public:
 	/**
 	 * Creates and initializes a new message context.
 	 *
+	 * This constructor overload is used for published and sent messages.
+	 *
 	 * @param InMessage The message payload.
 	 * @param InTypeInfo The message's type information.
 	 * @param InAttachment The binary data to attach to the message.
@@ -46,230 +56,58 @@ public:
 	{ }
 
 	/**
-	 * Creates and initializes a new message from an existing message context.
+	 * Creates and initializes a new message context from an existing context.
+	 *
+	 * This constructor overload is used for forwarded messages.
 	 *
-	 * @param InContext The context to initialize from.
+	 * @param InContext The existing context.
 	 * @param InForwarder The forwarder's address.
-	 * @param NewRecipients The message recipients.
-	 * @param NewScope The forwarding scope.
+	 * @param NewRecipients The recipients of the new context.
+	 * @param NewScope The message's new scope.
 	 * @param InTimeForwarded The time at which the message was forwarded.
-	 * @param InSenderThread The name of the thread from which the message was sent.
+	 * @param InForwarderThread The name of the thread from which the message was forwarded.
 	 */
-	FMessageContext( const IMessageContextRef& InContext, const FMessageAddress& InForwarder, const TArray<FMessageAddress>& NewRecipients, EMessageScope NewScope, const FDateTime& InTimeForwarded, ENamedThreads::Type InSenderThread )
-		: Attachment(nullptr)
-		, Forwarder(InForwarder)
-		, Message(nullptr)
+	FMessageContext( const IMessageContextRef& InContext, const FMessageAddress& InForwarder, const TArray<FMessageAddress>& NewRecipients, EMessageScope NewScope, const FDateTime& InTimeForwarded, ENamedThreads::Type InForwarderThread )
+		: Message(nullptr)
 		, OriginalContext(InContext)
 		, Recipients(NewRecipients)
 		, Scope(NewScope)
 		, Sender(InForwarder)
-		, SenderThread(InSenderThread)
-		, TimeForwarded(InTimeForwarded)
-		, TypeInfo(nullptr)
+		, SenderThread(InForwarderThread)
+		, TimeSent(InTimeForwarded)
 	{ }
 
 	/** Destructor. */
-	~FMessageContext()
-	{
-		if (Message != nullptr)
-		{
-			if (TypeInfo.IsValid())
-			{
-				TypeInfo->DestroyScriptStruct(Message);
-			}
-
-			FMemory::Free(Message);
-		}		
-	}
+	virtual ~FMessageContext() override;
 
 public:
 
 	// IMessageContext interface
 
-	virtual IMessageAttachmentPtr GetAttachment() const override
-	{
-		if (OriginalContext.IsValid())
-		{
-			return OriginalContext->GetAttachment();
-		}
-
-		return Attachment;
-	}
-
-	virtual const FDateTime& GetExpiration() const override
-	{
-		if (OriginalContext.IsValid())
-		{
-			return OriginalContext->GetExpiration();
-		}
-
-		return Expiration;
-	}
-
-	virtual const FMessageAddress& GetForwarder() const override
-	{
-		return Forwarder;
-	}
-
-	virtual const TMap<FName, FString>& GetHeaders() const override
-	{
-		if (OriginalContext.IsValid())
-		{
-			return OriginalContext->GetHeaders();
-		}
-
-		return Headers;
-	}
-
-	virtual const void* GetMessage() const override
-	{
-		if (OriginalContext.IsValid())
-		{
-			return OriginalContext->GetMessage();
-		}
-
-		return Message;
-	}
-
-	virtual FName GetMessageType() const override
-	{
-		if (IsValid())
-		{
-			return GetMessageTypeInfo()->GetFName();
-		}
-		
-		return NAME_None;
-	}
-
-	virtual const TWeakObjectPtr<UScriptStruct>& GetMessageTypeInfo() const override
-	{
-		if (OriginalContext.IsValid())
-		{
-			return OriginalContext->GetMessageTypeInfo();
-		}
-
-		return TypeInfo;
-	}
-
-	virtual IMessageContextPtr GetOriginalContext() const override
-	{
-		return OriginalContext;
-	}
-
-	virtual const TArray<FMessageAddress>& GetRecipients() const override
-	{
-		return Recipients;
-	}
-
-	virtual EMessageScope GetScope() const override
-	{
-		return Scope;
-	}
-
-	virtual const FMessageAddress& GetSender() const override
-	{
-		if (OriginalContext.IsValid())
-		{
-			return OriginalContext->GetSender();
-		}
-
-		return Sender;
-	}
-
-	virtual ENamedThreads::Type GetSenderThread() const override
-	{
-		return SenderThread;
-	}
-
-	virtual const FDateTime& GetTimeForwarded() const override
-	{
-		return TimeForwarded;
-	}
-
-	virtual const FDateTime& GetTimeSent() const override
-	{
-		if (OriginalContext.IsValid())
-		{
-			return OriginalContext->GetTimeSent();
-		}
-
-		return TimeSent;
-	}
-
-	virtual bool IsForwarded() const override
-	{
-		return OriginalContext.IsValid();
-	}
-
-	virtual bool IsValid() const override
-	{
-		if (OriginalContext.IsValid())
-		{
-			return OriginalContext->IsValid();
-		}
-
-		return ((Message != nullptr) && TypeInfo.IsValid(false, true));
-	}
-
-public:
-
-	// IMutableMessageContext interface
-
-	virtual void AddRecipient( const FMessageAddress& Recipient ) override
-	{
-		Recipients.Add(Recipient);
-	}
-
-	virtual void SetAttachment( const IMessageAttachmentPtr& InAttachment ) override
-	{
-		Attachment = InAttachment;
-	}
-
-	virtual void SetMessage( void* InMessage, UScriptStruct* InTypeInfo ) override
-	{
-		Message = InMessage;
-		TypeInfo = InTypeInfo;
-	}
-
-	virtual void SetExpiration( const FDateTime& InExpiration ) override
-	{
-		Expiration = InExpiration;
-	}
-
-	virtual void SetHeader( const FName& Key, const FString& Value ) override
-	{
-		Headers.Add(Key, Value);
-	}
-
-	virtual void SetScope( EMessageScope InScope ) override
-	{
-		Scope = InScope;
-	}
-
-	virtual void SetSender( const FMessageAddress& InSender ) override
-	{
-		Sender = InSender;
-	}
-
-	virtual void SetTimeSent( const FDateTime& InTimeSent ) override
-	{
-		TimeSent = InTimeSent;
-	}
+	virtual const TMap<FName, FString>& GetAnnotations() const override;
+	virtual IMessageAttachmentPtr GetAttachment() const override;
+	virtual const FDateTime& GetExpiration() const override;
+	virtual const void* GetMessage() const override;
+	virtual const TWeakObjectPtr<UScriptStruct>& GetMessageTypeInfo() const override;
+	virtual IMessageContextPtr GetOriginalContext() const override;
+	virtual const TArray<FMessageAddress>& GetRecipients() const override;
+	virtual EMessageScope GetScope() const override;
+	virtual const FMessageAddress& GetSender() const override;
+	virtual ENamedThreads::Type GetSenderThread() const override;
+	virtual const FDateTime& GetTimeForwarded() const override;
+	virtual const FDateTime& GetTimeSent() const override;
 
 private:
 
+	/** Holds the optional message annotations. */
+	TMap<FName, FString> Annotations;
+
 	/** Holds a pointer to attached binary data. */
 	IMessageAttachmentPtr Attachment;
 
 	/** Holds the expiration time. */
 	FDateTime Expiration;
 
-	/** Holds the address of the endpoint that forwarded this message. */
-	FMessageAddress Forwarder;
-
-	/** Holds the optional message headers. */
-	TMap<FName, FString> Headers;
-
 	/** Holds the message. */
 	void* Message;
 
@@ -288,9 +126,6 @@ private:
 	/** Holds the name of the thread from which the message was sent. */
 	ENamedThreads::Type SenderThread;
 
-	/** Holds the time at which the message was forwarded. */
-	FDateTime TimeForwarded;
-
 	/** Holds the time at which the message was sent. */
 	FDateTime TimeSent;
 
diff --git a/Engine/Source/Runtime/Messaging/Private/Bus/MessageRouter.h b/Engine/Source/Runtime/Messaging/Private/Bus/MessageRouter.h
--- a/Engine/Source/Runtime/Messaging/Private/Bus/MessageRouter.h
+++ b/Engine/Source/Runtime/Messaging/Private/Bus/MessageRouter.h
@@ -2,6 +2,13 @@
 
 #pragma once
 
+#include "IMessageContext.h"
+#include "IMessageInterceptor.h"
+#include "IMessageSubscription.h"
+#include "IMessageTracer.h"
+#include "IReceiveMessages.h"
+#include "ISendMessages.h"
+
 
 /**
  * Implements a topic-based message router.
diff --git a/Engine/Source/Runtime/Messaging/Private/Bus/MessageSubscription.h b/Engine/Source/Runtime/Messaging/Private/Bus/MessageSubscription.h
--- a/Engine/Source/Runtime/Messaging/Private/Bus/MessageSubscription.h
+++ b/Engine/Source/Runtime/Messaging/Private/Bus/MessageSubscription.h
@@ -2,6 +2,10 @@
 
 #pragma once
 
+#include "IMessageContext.h"
+#include "IMessageSubscription.h"
+#include "IReceiveMessages.h"
+
 
 /**
  * Implements a message subscription.
diff --git a/Engine/Source/Runtime/Messaging/Private/Bus/MessageTracer.h b/Engine/Source/Runtime/Messaging/Private/Bus/MessageTracer.h
--- a/Engine/Source/Runtime/Messaging/Private/Bus/MessageTracer.h
+++ b/Engine/Source/Runtime/Messaging/Private/Bus/MessageTracer.h
@@ -2,15 +2,12 @@
 
 #pragma once
 
-
-/** Type definition for weak pointers to instances of FMessageTracer. */
-typedef TWeakPtr<class FMessageTracer, ESPMode::ThreadSafe> FMessageTracerWeakPtr;
-
-/** Type definition for shared pointers to instances of FMessageTracer. */
-typedef TSharedPtr<class FMessageTracer, ESPMode::ThreadSafe> FMessageTracerPtr;
-
-/** Type definition for shared references to instances of FMessageTracer. */
-typedef TSharedRef<class FMessageTracer, ESPMode::ThreadSafe> FMessageTracerRef;
+#include "IMessageContext.h"
+#include "IMessageInterceptor.h"
+#include "IMessageSubscription.h"
+#include "IMessageTracer.h"
+#include "IMessageTracerBreakpoint.h"
+#include "IReceiveMessages.h"
 
 
 /**
@@ -40,6 +37,7 @@ class FMessageTracer
 	};
 
 public:
+
 	DECLARE_DELEGATE(TraceDelegate)
 
 	/** Default constructor. */
@@ -85,6 +83,7 @@ public:
 			return;
 		}
 
+		// @todo gmp: trace added subscriptions
 	}
 
 	/**
@@ -132,6 +131,7 @@ public:
 		{
 		}
 
+		// @todo gmp: trace intercepted messages
 	}
 
 	/**
@@ -147,6 +147,7 @@ public:
 			return;
 		}
 
+		// @todo gmp: trace removed interceptors
 	}
 
 	/**
@@ -161,6 +162,7 @@ public:
 			return;
 		}
 
+		// @todo gmp: trace removed recipients
 	}
 
 	/**
@@ -176,6 +178,7 @@ public:
 			return;
 		}
 
+		// @todo gmp: trace removed subscriptions
 	}
 
 	/**
@@ -404,6 +407,7 @@ protected:
 	 * Processes traces for sent messages.
 	 *
 	 * @param Context The context of the sent message.
+	 * @param SenderThread The name of the thread from which the message was sent.
 	 * @param TimeSecond The time at which the message was sent.
 	 */
 	void ProcessSentMessage( IMessageContextRef Context, double TimeSeconds );
@@ -466,3 +470,13 @@ private:
 	/** Holds a delegate that is executed when a new type has been added to the collection of known message types. */
 	FMessageTracerTypeAdded TypeAddedDelegate;
 };
+
+
+/** Type definition for weak pointers to instances of FMessageTracer. */
+typedef TWeakPtr<FMessageTracer, ESPMode::ThreadSafe> FMessageTracerWeakPtr;
+
+/** Type definition for shared pointers to instances of FMessageTracer. */
+typedef TSharedPtr<FMessageTracer, ESPMode::ThreadSafe> FMessageTracerPtr;
+
+/** Type definition for shared references to instances of FMessageTracer. */
+typedef TSharedRef<FMessageTracer, ESPMode::ThreadSafe> FMessageTracerRef;
diff --git a/Engine/Source/Runtime/Messaging/Private/MessagingModule.cpp b/Engine/Source/Runtime/Messaging/Private/MessagingModule.cpp
--- a/Engine/Source/Runtime/Messaging/Private/MessagingModule.cpp
+++ b/Engine/Source/Runtime/Messaging/Private/MessagingModule.cpp
@@ -1,6 +1,7 @@
 // Copyright 1998-2014 Epic Games, Inc. All Rights Reserved.
 
 #include "MessagingPrivatePCH.h"
+#include "IMessagingModule.h"
 #include "ModuleManager.h"
 
 
@@ -54,9 +55,9 @@ public:
 
 	// IMessagingModule interface
 
-	virtual IMessageBridgePtr CreateBridge( const FMessageAddress& Address, const IMessageBusRef& Bus, const ISerializeMessagesRef& Serializer, const IMessageTransportRef& Transport ) override
+	virtual IMessageBridgePtr CreateBridge( const FMessageAddress& Address, const IMessageBusRef& Bus, const IMessageTransportRef& Transport ) override
 	{
-		return MakeShareable(new FMessageBridge(Address, Bus, Serializer, Transport));
+		return MakeShareable(new FMessageBridge(Address, Bus, Transport));
 	}
 
 	virtual IMessageBusPtr CreateBus( const IAuthorizeMessageRecipientsPtr& RecipientAuthorizer ) override
@@ -64,12 +65,7 @@ public:
 		return MakeShareable(new FMessageBus(RecipientAuthorizer));
 	}
 
-	virtual ISerializeMessagesPtr CreateJsonMessageSerializer( ) override
-	{
-		return MakeShareable(new FJsonMessageSerializer());
-	}
-
-	virtual IMessageBusPtr GetDefaultBus( ) const override
+	virtual IMessageBusPtr GetDefaultBus() const override
 	{
 		return DefaultBus;
 	}
@@ -78,7 +74,7 @@ public:
 
 	// IModuleInterface interface
 
-	virtual void StartupModule( ) override
+	virtual void StartupModule() override
 	{
 #if PLATFORM_SUPPORTS_MESSAGEBUS
 		FCoreDelegates::OnPreExit.AddRaw(this, &FMessagingModule::HandleCorePreExit);
@@ -86,19 +82,19 @@ public:
 #endif	//PLATFORM_SUPPORTS_MESSAGEBUS
 	}
 
-	virtual void ShutdownModule( ) override
+	virtual void ShutdownModule() override
 	{
 		ShutdownDefaultBus();
 	}
 
-	virtual bool SupportsDynamicReloading( ) override
+	virtual bool SupportsDynamicReloading() override
 	{
 		return false;
 	}
 
 protected:
 
-	void ShutdownDefaultBus( )
+	void ShutdownDefaultBus()
 	{
 		if (!DefaultBus.IsValid())
 		{
@@ -125,7 +121,7 @@ protected:
 private:
 
 	// Callback for Core shutdown.
-	void HandleCorePreExit( )
+	void HandleCorePreExit()
 	{
 		ShutdownDefaultBus();
 	}
@@ -140,28 +136,10 @@ private:
 /* Misc static functions
  *****************************************************************************/
 
-TStatId FMessageForwardTask::GetStatId() const
-{
-	RETURN_QUICK_DECLARE_CYCLE_STAT(FMessageForwardTask, STATGROUP_TaskGraphTasks);
-}
-
-
 TStatId FMessageDispatchTask::GetStatId() const
 {
 	RETURN_QUICK_DECLARE_CYCLE_STAT(FMessageDispatchTask, STATGROUP_TaskGraphTasks);
 }
 
 
-TStatId FMessageSerializeTask::GetStatId() const
-{
-	RETURN_QUICK_DECLARE_CYCLE_STAT(FMessageSerializeTask, STATGROUP_TaskGraphTasks);
-}
-
-
-TStatId FMessageDeserializeTask::GetStatId() const
-{
-	RETURN_QUICK_DECLARE_CYCLE_STAT(FMessageDeserializeTask, STATGROUP_TaskGraphTasks);
-}
-
-
 IMPLEMENT_MODULE(FMessagingModule, Messaging);
diff --git a/Engine/Source/Runtime/Messaging/Private/MessagingPrivatePCH.h b/Engine/Source/Runtime/Messaging/Private/MessagingPrivatePCH.h
--- a/Engine/Source/Runtime/Messaging/Private/MessagingPrivatePCH.h
+++ b/Engine/Source/Runtime/Messaging/Private/MessagingPrivatePCH.h
@@ -2,13 +2,12 @@
 
 #pragma once
 
-#include "Messaging.h"
-
 
 /* Private dependencies
  *****************************************************************************/
 
-#include "Json.h"
+#include "CoreUObject.h"
+#include "TaskGraphInterfaces.h"
 
 
 /* Private includes
@@ -20,14 +19,5 @@
 #include "MessageDispatchTask.h"
 #include "MessageRouter.h"
 #include "MessageBus.h"
-
-#include "BsonMessageSerializer.h"
-#include "JsonMessageSerializer.h"
-#include "XmlMessageSerializer.h"
-
 #include "MessageAddressBook.h"
-#include "MessageData.h"
-#include "MessageDeserializeTask.h"
-#include "MessageForwardTask.h"
-#include "MessageSerializeTask.h"
 #include "MessageBridge.h"
diff --git a/Engine/Source/Runtime/Messaging/Private/Serialization/BsonMessageSerializer.h b/Engine/Source/Runtime/Messaging/Private/Serialization/BsonMessageSerializer.h
deleted file mode 100644
diff --git a/Engine/Source/Runtime/Messaging/Private/Serialization/JsonMessageSerializer.cpp b/Engine/Source/Runtime/Messaging/Private/Serialization/JsonMessageSerializer.cpp
deleted file mode 100644
diff --git a/Engine/Source/Runtime/Messaging/Private/Serialization/JsonMessageSerializer.h b/Engine/Source/Runtime/Messaging/Private/Serialization/JsonMessageSerializer.h
deleted file mode 100644
diff --git a/Engine/Source/Runtime/Messaging/Private/Serialization/XmlMessageSerializer.h b/Engine/Source/Runtime/Messaging/Private/Serialization/XmlMessageSerializer.h
deleted file mode 100644
diff --git a/Engine/Source/Runtime/Messaging/Public/Common/FileMessageAttachment.h b/Engine/Source/Runtime/Messaging/Public/Common/FileMessageAttachment.h
--- a/Engine/Source/Runtime/Messaging/Public/Common/FileMessageAttachment.h
+++ b/Engine/Source/Runtime/Messaging/Public/Common/FileMessageAttachment.h
@@ -2,6 +2,9 @@
 
 #pragma once
 
+#include "FileManager.h"
+#include "IMessageAttachment.h"
+
 
 /**
  * Implements a message attachment whose data is held in a file.
diff --git a/Engine/Source/Runtime/Messaging/Public/Common/MessageBridgeBuilder.h b/Engine/Source/Runtime/Messaging/Public/Common/MessageBridgeBuilder.h
--- a/Engine/Source/Runtime/Messaging/Public/Common/MessageBridgeBuilder.h
+++ b/Engine/Source/Runtime/Messaging/Public/Common/MessageBridgeBuilder.h
@@ -2,6 +2,12 @@
 
 #pragma once
 
+#include "IMessageBridge.h"
+#include "IMessageBus.h"
+#include "IMessageContext.h"
+#include "IMessageTransport.h"
+#include "IMessagingModule.h"
+
 
 /**
  * Implements a message bridge builder.
@@ -44,19 +50,6 @@ public:
 		return *this;
 	}
 
-	/**
-	 * Configures the bridge to use a custom message serializer.\
-	 *
-	 * @param CustomSerializer The custom serializer.
-	 * @return This instance (for method chaining).
-	 */
-	FMessageBridgeBuilder& UsingCustomSerializer( const ISerializeMessagesRef& CustomSerializer )
-	{
-		Serializer = CustomSerializer;
-
-		return *this;
-	}
-
 	/**
 	 * Configures the bridge to use a specific message transport technology.
 	 *
@@ -70,18 +63,6 @@ public:
 		return *this;
 	}
 
-	/**
-	 * Configures the bridge to use a binary message serializer.
-	 *
-	 * @return This instance (for method chaining).
-	 */
-	FMessageBridgeBuilder& UsingJsonSerializer()
-	{
-		Serializer = IMessagingModule::Get().CreateJsonMessageSerializer();
-
-		return *this;
-	}
-
 	/**
 	 * Sets the bridge's address.
 	 *
@@ -114,7 +95,7 @@ public:
 
 		if (Bus.IsValid())
 		{
-			Bridge = IMessagingModule::Get().CreateBridge(Address, Bus.ToSharedRef(), Serializer.ToSharedRef(), Transport.ToSharedRef());
+			Bridge = IMessagingModule::Get().CreateBridge(Address, Bus.ToSharedRef(), Transport.ToSharedRef());
 
 			if (Bridge.IsValid())
 			{
@@ -153,9 +134,6 @@ private:
 	/** Holds a flag indicating whether the bridge should be disabled. */
 	bool Disabled;
 
-	/** Holds a reference to the message serializer. */
-	ISerializeMessagesPtr Serializer;
-
 	/** Holds a reference to the message transport technology. */
 	IMessageTransportPtr Transport;
 };
diff --git a/Engine/Source/Runtime/Messaging/Public/Common/MessageEndpoint.h b/Engine/Source/Runtime/Messaging/Public/Common/MessageEndpoint.h
--- a/Engine/Source/Runtime/Messaging/Public/Common/MessageEndpoint.h
+++ b/Engine/Source/Runtime/Messaging/Public/Common/MessageEndpoint.h
@@ -161,7 +161,7 @@ public:
 	 *
 	 * @param NamedThread The name of the thread to receive messages on.
 	 */
-	void SetRecipientThread( ENamedThreads::Type NamedThread )
+	void SetRecipientThread( const ENamedThreads::Type& NamedThread )
 	{
 		RecipientThread = NamedThread;
 	}
@@ -174,7 +174,7 @@ public:
 	 * The message is effectively delivered again to this endpoint after the
 	 * original sent time plus the time delay have elapsed.
 	 *
-	 * @param Message The context of the message to defer.
+	 * @param Context The context of the message to defer.
 	 * @param Delay The time delay.
 	 */
 	void Defer( const IMessageContextRef& Context, const FTimespan& Delay )
@@ -183,7 +183,7 @@ public:
 
 		if (Bus.IsValid())
 		{
-			Bus->Forward(Context, TArrayBuilder<FMessageAddress>().Add(Address), Context->GetScope(), Delay, AsShared());
+			Bus->Forward(Context, TArrayBuilder<FMessageAddress>().Add(Address), Delay, AsShared());
 		}
 	}
 
@@ -194,16 +194,15 @@ public:
 	 *
 	 * @param Context The context of the message to forward.
 	 * @param Recipients The list of message recipients to forward the message to.
-	 * @param ForwardingScope The scope of the forwarded message.
 	 * @param Delay The time delay.
 	 */
-	void Forward( const IMessageContextRef& Context, const TArray<FMessageAddress>& Recipients, EMessageScope ForwardingScope, const FTimespan& Delay )
+	void Forward( const IMessageContextRef& Context, const TArray<FMessageAddress>& Recipients, const FTimespan& Delay )
 	{
 		IMessageBusPtr Bus = GetBusIfEnabled();
 
 		if (Bus.IsValid())
 		{
-			Bus->Forward(Context, Recipients, ForwardingScope, Delay, AsShared());
+			Bus->Forward(Context, Recipients, Delay, AsShared());
 		}
 	}
 
@@ -447,11 +446,10 @@ public:
 	 *
 	 * @param Context The context of the message to forward.
 	 * @param Recipient The address of the recipient to forward the message to.
-	 * @param ForwardingScope The scope of the forwarded message.
 	 */
-	void Forward( const IMessageContextRef& Context, const FMessageAddress& Recipient, EMessageScope ForwardingScope )
+	void Forward( const IMessageContextRef& Context, const FMessageAddress& Recipient )
 	{
-		Forward(Context, TArrayBuilder<FMessageAddress>().Add(Recipient), ForwardingScope, FTimespan::Zero());
+		Forward(Context, TArrayBuilder<FMessageAddress>().Add(Recipient), FTimespan::Zero());
 	}
 
 	/**
@@ -464,9 +462,9 @@ public:
 	 * @param ForwardingScope The scope of the forwarded message.
 	 * @param Delay The delay after which to publish the message.
 	 */
-	void Forward( const IMessageContextRef& Context, const FMessageAddress& Recipient, EMessageScope ForwardingScope, const FTimespan& Delay )
+	void Forward( const IMessageContextRef& Context, const FMessageAddress& Recipient, const FTimespan& Delay )
 	{
-		Forward(Context, TArrayBuilder<FMessageAddress>().Add(Recipient), ForwardingScope, Delay);
+		Forward(Context, TArrayBuilder<FMessageAddress>().Add(Recipient), Delay);
 	}
 
 	/**
@@ -478,9 +476,9 @@ public:
 	 * @param Recipients The list of message recipients to forward the message to.
 	 * @param ForwardingScope The scope of the forwarded message.
 	 */
-	void Forward( const IMessageContextRef& Context, const TArray<FMessageAddress>& Recipients, EMessageScope ForwardingScope )
+	void Forward( const IMessageContextRef& Context, const TArray<FMessageAddress>& Recipients )
 	{
-		Forward(Context, Recipients, ForwardingScope, FTimespan::Zero());
+		Forward(Context, Recipients, FTimespan::Zero());
 	}
 
 	/**
diff --git a/Engine/Source/Runtime/Messaging/Public/Common/MessageEndpointBuilder.h b/Engine/Source/Runtime/Messaging/Public/Common/MessageEndpointBuilder.h
--- a/Engine/Source/Runtime/Messaging/Public/Common/MessageEndpointBuilder.h
+++ b/Engine/Source/Runtime/Messaging/Public/Common/MessageEndpointBuilder.h
@@ -2,8 +2,11 @@
 
 #pragma once
 
+#include "IMessageBus.h"
+#include "IMessageHandler.h"
 #include "IMessagingModule.h"
 #include "MessageEndpoint.h"
+#include "TaskGraphInterfaces.h"
 
 
 /**
diff --git a/Engine/Source/Runtime/Messaging/Public/Interfaces/IMessageBus.h b/Engine/Source/Runtime/Messaging/Public/Interfaces/IMessageBus.h
--- a/Engine/Source/Runtime/Messaging/Public/Interfaces/IMessageBus.h
+++ b/Engine/Source/Runtime/Messaging/Public/Interfaces/IMessageBus.h
@@ -3,15 +3,17 @@
 #pragma once
 
 
+// forward declarations
 enum class EMessageScope;
 struct FMessageAddress;
-class IMessageInterceptor;
 class IMessageAttachment;
 class IMessageContext;
+class IMessageInterceptor;
 class IMessageSubscription;
 class IMessageTracer;
 class IReceiveMessages;
 class ISendMessages;
+class UScriptStruct;
 
 
 /** Delegate type for message bus shutdowns. */
@@ -108,12 +110,11 @@ public:
 	 *
 	 * @param Context The context of the message to forward.
 	 * @param Recipients The list of message recipients to forward the message to.
-	 * @param NewScope The forwarding scope.
 	 * @param Delay The deferral time.
 	 * @param Forwarder The sender that forwards the message.
 	 * @see Publish, Send
 	 */
-	virtual void Forward( const TSharedRef<IMessageContext, ESPMode::ThreadSafe>& Context, const TArray<FMessageAddress>& Recipients, EMessageScope NewScope, const FTimespan& Delay, const TSharedRef<ISendMessages, ESPMode::ThreadSafe>& Forwarder ) = 0;
+	virtual void Forward( const TSharedRef<IMessageContext, ESPMode::ThreadSafe>& Context, const TArray<FMessageAddress>& Recipients, const FTimespan& Delay, const TSharedRef<ISendMessages, ESPMode::ThreadSafe>& Forwarder ) = 0;
 
 	/**
 	 * Gets the message bus tracer.
diff --git a/Engine/Source/Runtime/Messaging/Public/Interfaces/IMessageContext.h b/Engine/Source/Runtime/Messaging/Public/Interfaces/IMessageContext.h
--- a/Engine/Source/Runtime/Messaging/Public/Interfaces/IMessageContext.h
+++ b/Engine/Source/Runtime/Messaging/Public/Interfaces/IMessageContext.h
@@ -2,12 +2,16 @@
 
 #pragma once
 
+#include "TaskGraphInterfaces.h"
 
+
+// forward declarations
 class IMessageAttachment;
+class UScriptStruct;
 
 
 /**
- * Implements a message endpoint address.
+ * Structure for message endpoint addresses.
  */
 struct FMessageAddress
 {
@@ -166,11 +170,12 @@ typedef TRangeBound<EMessageScope> FMessageScopeRangeBound;
 /**
  * Interface for message contexts.
  *
- * A message context is created internally by the message bus to store additional information about a
- * message needed to describe and correctly route the message to recipients. The two most common aspects
- * of message contexts used by recipients are the sender address and message attachments.
+ * Messages are delivered inside message contexts, which store the message itself plus additional data
+ * associated with the message. Recipients of a message are usually interested in data that describes the
+ * message, such as its origin or when it expires. They may also be interested in optional out-of-band
+ * binary data that is attached to the message.
  *
- * The sender address (@see GetSender) is often needed to send a reply message to a message sender, i.e.
+ * The sender's address (@see GetSender) is often needed to send a reply message to a message sender, i.e.
  * in response to a published message. The message attachment (@see GetAttchment) is an optional bundle
  * of binary bulk data that is transferred independently from the message itself and allows for transferring
  * larger amounts of data that would otherwise clog up the messaging system.
@@ -184,6 +189,13 @@ class IMessageContext
 {
 public:
 
+	/**
+	 * Gets the optional message annotations.
+	 *
+	 * @return Message header collection.
+	 */
+	virtual const TMap<FName, FString>& GetAnnotations() const = 0;
+
 	/**
 	 * Gets the message attachment, if present.
 	 *
@@ -199,59 +211,39 @@ public:
 	virtual const FDateTime& GetExpiration() const = 0;
 
 	/**
-	 * Gets the message address of the endpoint that forwarded this message.
-	 *
-	 * @return The forwarder's address.
-	 * @see IsForwarded
-	 */
-	virtual const FMessageAddress& GetForwarder() const = 0;
-
-	/**
-	 * Gets the optional message headers.
-	 *
-	 * @return Message header collection.
-	 */
-	virtual const TMap<FName, FString>& GetHeaders() const = 0;
-
-	/**
-	 * Gets the message object.
+	 * Gets the message data.
 	 *
-	 * @return The message.
+	 * @return A pointer to the message data.
+	 * @see GetMessageType, GetMessageTypeInfo
 	 */
 	virtual const void* GetMessage() const = 0;
 
-	/**
-	 * Gets the name of the message type.
-	 *
-	 * @return Message type name.
-	 */
-	virtual FName GetMessageType() const = 0;
-
 	/**
 	 * Gets the message's type information.
 	 *
 	 * @return Message type information.
+	 * @see GetMessage, GetMessageType
 	 */
 	virtual const TWeakObjectPtr<UScriptStruct>& GetMessageTypeInfo() const = 0;
 
 	/**
 	 * Returns the original message context in case the message was forwarded.
 	 *
-	 * @return The original message.
+	 * @return The original message context, or nullptr if the message wasn't forwarded.
 	 */
 	virtual TSharedPtr<IMessageContext, ESPMode::ThreadSafe> GetOriginalContext() const = 0;
 
 	/**
-	* Gets the list of message recipients.
-	*
-	* @return Message recipients.
+	 * Gets the list of message recipients.
+	 *
+	 * @return Message recipients.
 	 */
 	virtual const TArray<FMessageAddress>& GetRecipients() const = 0;
 
 	/**
-	 * Gets the message's scope.
+	 * Gets the scope to which the message was sent.
 	 *
-	 * @return Scope.
+	 * @return The message scope.
 	 */
 	virtual EMessageScope GetScope() const = 0;
 
@@ -283,20 +275,44 @@ public:
 	 */
 	virtual const FDateTime& GetTimeSent() const = 0;
 
+public:
+
+	/**
+	 * Gets the name of the message type.
+	 *
+	 * @return Message type name.
+	 */
+	FName GetMessageType() const
+	{
+		if (IsValid())
+		{
+			return GetMessageTypeInfo()->GetFName();
+		}
+		
+		return NAME_None;
+	}
+
 	/**
 	 * Checks whether this is a forwarded message.
 	 *
 	 * @return true if the message was forwarded, false otherwise.
-	 * @see GetForwarder
+	 * @see GetOriginalContext, IsValid
 	 */
-	virtual bool IsForwarded() const = 0;
+	bool IsForwarded() const
+	{
+		return GetOriginalContext().IsValid();
+	}
 
 	/**
 	 * Checks whether this context is valid.
 	 *
 	 * @return true if the context is valid, false otherwise.
+	 * @see IsForwarded
 	 */
-	virtual bool IsValid() const = 0;
+	bool IsValid() const
+	{
+		return ((GetMessage() != nullptr) && GetMessageTypeInfo().IsValid(false, true));
+	}
 
 public:
 
diff --git a/Engine/Source/Runtime/Messaging/Public/Interfaces/IMessageData.h b/Engine/Source/Runtime/Messaging/Public/Interfaces/IMessageData.h
deleted file mode 100644
diff --git a/Engine/Source/Runtime/Messaging/Public/Interfaces/IMessageHandler.h b/Engine/Source/Runtime/Messaging/Public/Interfaces/IMessageHandler.h
--- a/Engine/Source/Runtime/Messaging/Public/Interfaces/IMessageHandler.h
+++ b/Engine/Source/Runtime/Messaging/Public/Interfaces/IMessageHandler.h
@@ -2,7 +2,10 @@
 
 #pragma once
 
+#include "CoreUObject.h"
 
+
+// forward declarations
 class IMessageContext;
 
 
diff --git a/Engine/Source/Runtime/Messaging/Public/Interfaces/IMessageTracer.h b/Engine/Source/Runtime/Messaging/Public/Interfaces/IMessageTracer.h
--- a/Engine/Source/Runtime/Messaging/Public/Interfaces/IMessageTracer.h
+++ b/Engine/Source/Runtime/Messaging/Public/Interfaces/IMessageTracer.h
@@ -2,7 +2,9 @@
 
 #pragma once
 
-#include "IMessageContext.h"
+
+// forward declarations
+class IMessageContext;
 
 
 /** Type definition for shared pointers to instances of FMessageTracerAddressInfo. */
diff --git a/Engine/Source/Runtime/Messaging/Public/Interfaces/IMessageTracerBreakpoint.h b/Engine/Source/Runtime/Messaging/Public/Interfaces/IMessageTracerBreakpoint.h
--- a/Engine/Source/Runtime/Messaging/Public/Interfaces/IMessageTracerBreakpoint.h
+++ b/Engine/Source/Runtime/Messaging/Public/Interfaces/IMessageTracerBreakpoint.h
@@ -3,6 +3,7 @@
 #pragma once
 
 
+// forward declarations
 class IMessageContext;
 
 
@@ -25,7 +26,7 @@ public:
 	/**
 	 * Checks whether the tracer should break on the given message.
 	 *
-	 * @param Context The message to break on.
+	 * @param Context The context of the message to break on.
 	 * @return true if the tracer should break, false otherwise.
 	 */
 	virtual bool ShouldBreak( const TSharedRef<IMessageContext, ESPMode::ThreadSafe>& Context ) const = 0;
diff --git a/Engine/Source/Runtime/Messaging/Public/Interfaces/IMessageTransport.h b/Engine/Source/Runtime/Messaging/Public/Interfaces/IMessageTransport.h
--- a/Engine/Source/Runtime/Messaging/Public/Interfaces/IMessageTransport.h
+++ b/Engine/Source/Runtime/Messaging/Public/Interfaces/IMessageTransport.h
@@ -3,6 +3,7 @@
 #pragma once
 
 
+// forward declarations
 class IMessageAttachment;
 class IMessageData;
 
@@ -10,11 +11,10 @@ class IMessageData;
 /**
  * Delegate type for receiving transport messages.
  *
- * The first parameter is the serialized message content.
- * The second parameter is the message attachment, if any.
- * The third parameter is the identifier of the transport node that sent the message.
+ * The first parameter is the context of the received message.
+ * The second parameter is the identifier of the transport node that sent the message.
  */
-DECLARE_DELEGATE_ThreeParams(FOnMessageTransportMessageReceived, FArchive&, const IMessageAttachmentPtr&, const FGuid&);
+DECLARE_DELEGATE_TwoParams(FOnMessageTransportMessageReceived, const IMessageContextRef&, const FGuid&);
 
 /**
  * Delegate type for handling the discovery of transport nodes.
@@ -69,12 +69,11 @@ public:
 	/**
 	 * Transports the given message data to the specified network nodes.
 	 *
-	 * @param Data The serialized message data to transport.
-	 * @param Attachment An optional message attachment (i.e. file or memory buffer).
+	 * @param Context The context of the message to transport.
 	 * @param Recipients The transport nodes to send the message to.
 	 * @return true if the message is being transported, false otherwise.
 	 */
-	virtual bool TransportMessage( const TSharedRef<IMessageData, ESPMode::ThreadSafe>& Data, const TSharedPtr<IMessageAttachment, ESPMode::ThreadSafe>& Attachment, const TArray<FGuid>& Recipients ) = 0;
+	virtual bool TransportMessage( const IMessageContextRef& Context, const TArray<FGuid>& Recipients ) = 0;
 
 public:
 
diff --git a/Engine/Source/Runtime/Messaging/Public/Interfaces/IMessagingModule.h b/Engine/Source/Runtime/Messaging/Public/Interfaces/IMessagingModule.h
--- a/Engine/Source/Runtime/Messaging/Public/Interfaces/IMessagingModule.h
+++ b/Engine/Source/Runtime/Messaging/Public/Interfaces/IMessagingModule.h
@@ -5,11 +5,11 @@
 #include "ModuleInterface.h"
 
 
+// forward declarations
 struct FMessageAddress;
 class IAuthorizeMessageRecipients;
 class IMessageBus;
 class IMessageBridge;
-class ISerializeMessages;
 class IMessageTransport;
 
 
@@ -31,12 +31,11 @@ public:
 	 *
 	 * @param Address The bridge's address on the message bus.
 	 * @param Bus The message bus to attach the bridge to.
-	 * @param Serializer The message serializer to use.
 	 * @param Transport The message transport technology to use.
 	 * @return The new message bridge, or nullptr if the bridge couldn't be created.
 	 * @see CreateBus
 	 */
-	virtual TSharedPtr<IMessageBridge, ESPMode::ThreadSafe> CreateBridge( const FMessageAddress& Address, const TSharedRef<IMessageBus, ESPMode::ThreadSafe>& Bus, const TSharedRef<ISerializeMessages, ESPMode::ThreadSafe>& Serializer, const TSharedRef<IMessageTransport, ESPMode::ThreadSafe>& Transport ) = 0;
+	virtual TSharedPtr<IMessageBridge, ESPMode::ThreadSafe> CreateBridge( const FMessageAddress& Address, const TSharedRef<IMessageBus, ESPMode::ThreadSafe>& Bus, const TSharedRef<IMessageTransport, ESPMode::ThreadSafe>& Transport ) = 0;
 
 	/**
 	 * Creates a new message bus.
@@ -47,13 +46,6 @@ public:
 	 */
 	virtual TSharedPtr<IMessageBus, ESPMode::ThreadSafe> CreateBus( const TSharedPtr<IAuthorizeMessageRecipients>& RecipientAuthorizer ) = 0;
 
-	/**
-	 * Creates a Json message serializer (deprecated).
-	 *
-	 * @return A new serializer.
-	 */
-	virtual TSharedPtr<ISerializeMessages, ESPMode::ThreadSafe> CreateJsonMessageSerializer() = 0;
-
 	/**
 	 * Gets the default message bus if it has been initialized.
 	 *
diff --git a/Engine/Source/Runtime/Messaging/Public/Interfaces/ISerializeMessages.h b/Engine/Source/Runtime/Messaging/Public/Interfaces/ISerializeMessages.h
deleted file mode 100644
diff --git a/Engine/Source/Runtime/Messaging/Public/Messaging.h b/Engine/Source/Runtime/Messaging/Public/Messaging.h
--- a/Engine/Source/Runtime/Messaging/Public/Messaging.h
+++ b/Engine/Source/Runtime/Messaging/Public/Messaging.h
@@ -7,15 +7,12 @@
  *****************************************************************************/
 
 #include "CoreUObject.h"
-#include "ModuleManager.h"
-#include "TaskGraphInterfaces.h"
 
 
 /* Public Includes
  *****************************************************************************/
 
 #include "IMessageAttachment.h"
-#include "IMessageData.h"
 #include "IMessageContext.h"
 #include "IMessageInterceptor.h"
 #include "IMessageHandler.h"
@@ -27,7 +24,6 @@
 #include "IMessageTracer.h"
 #include "IMessageBus.h"
 #include "IMutableMessageContext.h"
-#include "ISerializeMessages.h"
 #include "IMessageTransport.h"
 #include "IMessageBridge.h"
 #include "IMessagingModule.h"
