UdpMessaging: Made MulticastSocket creation optional

A multicast socket is not strictly required for the UDP transport to work. For example, when adding devices as static endpoints in the Editor, all communication will be handled through the UnicastSocket alone. This is a temporary workaround for platforms that may not support multicast. A proper fallback to broadcast and then unicast should be implemented in the future.

diff --git a/Engine/Plugins/Messaging/UdpMessaging/Source/UdpMessaging/Private/Transport/UdpMessageTransport.cpp b/Engine/Plugins/Messaging/UdpMessaging/Source/UdpMessaging/Private/Transport/UdpMessageTransport.cpp
--- a/Engine/Plugins/Messaging/UdpMessaging/Source/UdpMessaging/Private/Transport/UdpMessageTransport.cpp
+++ b/Engine/Plugins/Messaging/UdpMessaging/Source/UdpMessaging/Private/Transport/UdpMessageTransport.cpp
@@ -32,53 +32,49 @@ FUdpMessageTransport::~FUdpMessageTransport()
 
 bool FUdpMessageTransport::StartTransport()
 {
-	// create network sockets
-	MulticastSocket = FUdpSocketBuilder(TEXT("UdpMessageMulticastSocket"))
+	// create & initialize unicast socket
+	UnicastSocket = FUdpSocketBuilder(TEXT("UdpMessageUnicastSocket"))
 		.AsNonBlocking()
-		.AsReusable()
-		.BoundToAddress(LocalEndpoint.GetAddress())
-		.BoundToPort(MulticastEndpoint.GetPort())
-		.JoinedToGroup(MulticastEndpoint.GetAddress())
 		.WithMulticastLoopback()
-		.WithMulticastTtl(MulticastTtl);
+		.BoundToEndpoint(LocalEndpoint)
+		.WithReceiveBufferSize(UDP_MESSAGING_RECEIVE_BUFFER_SIZE);
 
-	if (MulticastSocket == nullptr)
+	if (UnicastSocket == nullptr)
 	{
-		GLog->Logf(TEXT("UdpMessageTransport.StartTransport: Failed to create multicast socket on %s, joined to %s with TTL %i"), *LocalEndpoint.ToText().ToString(), *MulticastEndpoint.ToText().ToString(), MulticastTtl);
+		GLog->Logf(TEXT("UdpMessageTransport.StartTransport: Failed to create unicast socket on %s"), *LocalEndpoint.ToText().ToString());
 
 		return false;
 	}
 
-	UnicastSocket = FUdpSocketBuilder(TEXT("UdpMessageUnicastSocket"))
+	// create & initialize multicast socket (optional)
+	MulticastSocket = FUdpSocketBuilder(TEXT("UdpMessageMulticastSocket"))
 		.AsNonBlocking()
+		.AsReusable()
+		.BoundToAddress(LocalEndpoint.GetAddress())
+		.BoundToPort(MulticastEndpoint.GetPort())
+		.JoinedToGroup(MulticastEndpoint.GetAddress())
 		.WithMulticastLoopback()
-		.BoundToEndpoint(LocalEndpoint);
+		.WithMulticastTtl(MulticastTtl)
+		.WithReceiveBufferSize(UDP_MESSAGING_RECEIVE_BUFFER_SIZE);
 
-	if (UnicastSocket == nullptr)
+	if (MulticastSocket == nullptr)
 	{
-		SocketSubsystem->DestroySocket(MulticastSocket);
-		MulticastSocket = nullptr;
-
-		GLog->Logf(TEXT("UdpMessageTransport.StartTransport: Failed to create unicast socket on %s"), *LocalEndpoint.ToText().ToString());
-
-		return false;
+		GLog->Logf(TEXT("UdpMessageTransport.StartTransport: Failed to create multicast socket on %s, joined to %s with TTL %i"), *LocalEndpoint.ToText().ToString(), *MulticastEndpoint.ToText().ToString(), MulticastTtl);
 	}
 
-	// adjust buffer sizes
-	int32 NewSize = 0;
-	MulticastSocket->SetReceiveBufferSize(2 * 1024 * 1024, NewSize);
-	UnicastSocket->SetReceiveBufferSize(2 * 1024 * 1024, NewSize);
-
+	// initialize threads
 	FTimespan ThreadWaitTime = FTimespan::FromMilliseconds(100);
 
-	// initialize threads
 	MessageProcessor = new FUdpMessageProcessor(UnicastSocket, FGuid::NewGuid(), MulticastEndpoint);
 	MessageProcessor->OnMessageReassembled().BindRaw(this, &FUdpMessageTransport::HandleProcessorMessageReassembled);
 	MessageProcessor->OnNodeDiscovered().BindRaw(this, &FUdpMessageTransport::HandleProcessorNodeDiscovered);
 	MessageProcessor->OnNodeLost().BindRaw(this, &FUdpMessageTransport::HandleProcessorNodeLost);
 
-	MulticastReceiver = new FUdpSocketReceiver(MulticastSocket, ThreadWaitTime, TEXT("UdpMessageMulticastReceiver"));
-	MulticastReceiver->OnDataReceived().BindRaw(this, &FUdpMessageTransport::HandleSocketDataReceived);
+	if (MulticastSocket != nullptr)
+	{
+		MulticastReceiver = new FUdpSocketReceiver(MulticastSocket, ThreadWaitTime, TEXT("UdpMessageMulticastReceiver"));
+		MulticastReceiver->OnDataReceived().BindRaw(this, &FUdpMessageTransport::HandleSocketDataReceived);
+	}
 
 	UnicastReceiver = new FUdpSocketReceiver(UnicastSocket, ThreadWaitTime, TEXT("UdpMessageUnicastReceiver"));
 	UnicastReceiver->OnDataReceived().BindRaw(this, &FUdpMessageTransport::HandleSocketDataReceived);
diff --git a/Engine/Plugins/Messaging/UdpMessaging/Source/UdpMessaging/Private/UdpMessagingPrivatePCH.h b/Engine/Plugins/Messaging/UdpMessaging/Source/UdpMessaging/Private/UdpMessagingPrivatePCH.h
--- a/Engine/Plugins/Messaging/UdpMessaging/Source/UdpMessaging/Private/UdpMessagingPrivatePCH.h
+++ b/Engine/Plugins/Messaging/UdpMessaging/Source/UdpMessaging/Private/UdpMessagingPrivatePCH.h
@@ -26,6 +26,9 @@
 /** Defines the maximum number of recipients a message can have. */
 #define UDP_MESSAGING_MAX_RECIPIENTS 1024
 
+/** Defines the desired size of socket receive buffers (in bytes). */
+#define UDP_MESSAGING_RECEIVE_BUFFER_SIZE 2 * 1024 * 1024
+
 /** Defines the protocol version of the UDP message transport. */
 #define UDP_MESSAGING_TRANSPORT_PROTOCOL_VERSION 10
 
