Fix for FUdpDeserializedMessage::Deserialize allowing negative recipient array size (UE-11794)

diff --git a/Engine/Plugins/Messaging/UdpMessaging/Source/UdpMessaging/Private/Transport/UdpDeserializedMessage.cpp b/Engine/Plugins/Messaging/UdpMessaging/Source/UdpMessaging/Private/Transport/UdpDeserializedMessage.cpp
--- a/Engine/Plugins/Messaging/UdpMessaging/Source/UdpMessaging/Private/Transport/UdpDeserializedMessage.cpp
+++ b/Engine/Plugins/Messaging/UdpMessaging/Source/UdpMessaging/Private/Transport/UdpDeserializedMessage.cpp
@@ -39,7 +39,7 @@ bool FUdpDeserializedMessage::Deserialize( const FUdpReassembledMessageRef& Reas
 		int32 NumRecipients = 0;
 		MessageReader << NumRecipients;
 
-		if (NumRecipients > UDP_MESSAGING_MAX_RECIPIENTS)
+		if ((NumRecipients < 0) || (NumRecipients > UDP_MESSAGING_MAX_RECIPIENTS))
 		{
 			return false;
 		}
