Added missing dynamic module dependencies.

diff --git a/Engine/Plugins/Messaging/UdpMessaging/Source/UdpMessaging/UdpMessaging.Build.cs b/Engine/Plugins/Messaging/UdpMessaging/Source/UdpMessaging/UdpMessaging.Build.cs
--- a/Engine/Plugins/Messaging/UdpMessaging/Source/UdpMessaging/UdpMessaging.Build.cs
+++ b/Engine/Plugins/Messaging/UdpMessaging/Source/UdpMessaging/UdpMessaging.Build.cs
@@ -6,6 +6,13 @@ namespace UnrealBuildTool.Rules
 	{
 		public UdpMessaging(TargetInfo Target)
 		{
+            DynamicallyLoadedModuleNames.AddRange(
+                new string[] {
+                    "Messaging",
+					"Settings",
+				}
+            );
+
 			PublicDependencyModuleNames.AddRange(
 				new string[]
 				{
