Messaging: Own plug-in category for messaging plug-ins

diff --git a/Engine/Plugins/Messaging/MessagingDebugger/MessagingDebugger.uplugin b/Engine/Plugins/Messaging/MessagingDebugger/MessagingDebugger.uplugin
--- a/Engine/Plugins/Messaging/MessagingDebugger/MessagingDebugger.uplugin
+++ b/Engine/Plugins/Messaging/MessagingDebugger/MessagingDebugger.uplugin
@@ -8,7 +8,7 @@
 	"CreatedByURL" : "http://epicgames.com",
 	"EngineVersion" : "4.2.0",
 	"Description" : "Provides a visual debugger for the messaging sub-system.",
-	"Category" : "Misc",
+	"Category" : "Messaging",
 	"EnabledByDefault" : true,
 
 	"Modules" :
diff --git a/Engine/Plugins/Messaging/UdpMessaging/UdpMessaging.uplugin b/Engine/Plugins/Messaging/UdpMessaging/UdpMessaging.uplugin
--- a/Engine/Plugins/Messaging/UdpMessaging/UdpMessaging.uplugin
+++ b/Engine/Plugins/Messaging/UdpMessaging/UdpMessaging.uplugin
@@ -8,7 +8,7 @@
 	"CreatedByURL" : "http://epicgames.com",
 	"EngineVersion" : "4.2.0",
 	"Description" : "Adds a UDP based transport and tunneling layer to the messaging sub-system for sending and receiving messages between networked computers and devices.",
-	"Category" : "Misc",
+	"Category" : "Messaging",
 	"EnabledByDefault" : true,
 
 	"Modules" :
