UdpMessaging: Always forcing service restart on settings changes

diff --git a/Engine/Plugins/Messaging/UdpMessaging/Source/UdpMessaging/Private/UdpMessagingModule.cpp b/Engine/Plugins/Messaging/UdpMessaging/Source/UdpMessaging/Private/UdpMessagingModule.cpp
--- a/Engine/Plugins/Messaging/UdpMessaging/Source/UdpMessaging/Private/UdpMessagingModule.cpp
+++ b/Engine/Plugins/Messaging/UdpMessaging/Source/UdpMessaging/Private/UdpMessagingModule.cpp
@@ -196,10 +196,7 @@ protected:
 
 		if (Settings.EnableTransport)
 		{
-			if (!MessageBridge.IsValid())
-			{
-				InitializeBridge();
-			}
+			InitializeBridge();
 		}
 		else
 		{
@@ -208,10 +205,7 @@ protected:
 
 		if (Settings.EnableTunnel)
 		{
-			if (!MessageTunnel.IsValid())
-			{
-				InitializeTunnel();
-			}
+			InitializeTunnel();
 		}
 		else
 		{
