UdpMessaging: Added log category for module

diff --git a/Engine/Plugins/Messaging/UdpMessaging/Source/UdpMessaging/Private/Transport/UdpMessageProcessor.cpp b/Engine/Plugins/Messaging/UdpMessaging/Source/UdpMessaging/Private/Transport/UdpMessageProcessor.cpp
--- a/Engine/Plugins/Messaging/UdpMessaging/Source/UdpMessaging/Private/Transport/UdpMessageProcessor.cpp
+++ b/Engine/Plugins/Messaging/UdpMessaging/Source/UdpMessaging/Private/Transport/UdpMessageProcessor.cpp
@@ -26,18 +26,18 @@ FUdpMessageProcessor::FUdpMessageProcessor( FSocket* InSocket, const FGuid& InNo
 
 	const UUdpMessagingSettings& Settings = *GetDefault<UUdpMessagingSettings>();
 
-	for (int32 StaticEndpointIndex = 0; StaticEndpointIndex < Settings.StaticEndpoints.Num(); ++StaticEndpointIndex)
+	for (auto& StaticEndpoint : Settings.StaticEndpoints)
 	{
 		FIPv4Endpoint Endpoint;
 
-		if (FIPv4Endpoint::Parse(Settings.StaticEndpoints[StaticEndpointIndex], Endpoint))
+		if (FIPv4Endpoint::Parse(StaticEndpoint, Endpoint))
 		{
 			FNodeInfo& NodeInfo = StaticNodes.FindOrAdd(Endpoint);
 			NodeInfo.Endpoint = Endpoint;
 		}
 		else
 		{
-			GLog->Logf(TEXT("Warning: Invalid UDP Messaging StaticNode '%s'"), *Settings.StaticEndpoints[StaticEndpointIndex]);
+			UE_LOG(LogUdpMessaging, Warning, TEXT("Invalid UDP Messaging Static Endpoint '%s'"), *StaticEndpoint);
 		}
 	}
 }
@@ -406,7 +406,7 @@ void FUdpMessageProcessor::ProcessTimeoutSegment( FInboundSegment& Segment, FNod
 
 void FUdpMessageProcessor::ProcessUnknownSegment( FInboundSegment& Segment, FNodeInfo& EndpointInfo, uint8 SegmentType )
 {
-	GLog->Logf(TEXT("FUdpMessageProcessor: Received unknown segment type '%i' from %s"), SegmentType, *Segment.Sender.ToText().ToString());
+	UE_LOG(LogUdpMessaging, Verbose, TEXT("Received unknown segment type '%i' from %s"), SegmentType, *Segment.Sender.ToText().ToString());
 }
 
 
diff --git a/Engine/Plugins/Messaging/UdpMessaging/Source/UdpMessaging/Private/Transport/UdpMessageTransport.cpp b/Engine/Plugins/Messaging/UdpMessaging/Source/UdpMessaging/Private/Transport/UdpMessageTransport.cpp
--- a/Engine/Plugins/Messaging/UdpMessaging/Source/UdpMessaging/Private/Transport/UdpMessageTransport.cpp
+++ b/Engine/Plugins/Messaging/UdpMessaging/Source/UdpMessaging/Private/Transport/UdpMessageTransport.cpp
@@ -41,7 +41,7 @@ bool FUdpMessageTransport::StartTransport()
 
 	if (UnicastSocket == nullptr)
 	{
-		GLog->Logf(TEXT("UdpMessageTransport.StartTransport: Failed to create unicast socket on %s"), *LocalEndpoint.ToText().ToString());
+		UE_LOG(LogUdpMessaging, Error, TEXT("StartTransport failed to create unicast socket on %s"), *LocalEndpoint.ToText().ToString());
 
 		return false;
 	}
@@ -59,7 +59,7 @@ bool FUdpMessageTransport::StartTransport()
 
 	if (MulticastSocket == nullptr)
 	{
-		GLog->Logf(TEXT("UdpMessageTransport.StartTransport: Failed to create multicast socket on %s, joined to %s with TTL %i"), *LocalEndpoint.ToText().ToString(), *MulticastEndpoint.ToText().ToString(), MulticastTtl);
+		UE_LOG(LogUdpMessaging, Error, TEXT("StartTransport failed to create multicast socket on %s, joined to %s with TTL %i"), *LocalEndpoint.ToText().ToString(), *MulticastEndpoint.ToText().ToString(), MulticastTtl);
 	}
 
 	// initialize threads
diff --git a/Engine/Plugins/Messaging/UdpMessaging/Source/UdpMessaging/Private/UdpMessagingModule.cpp b/Engine/Plugins/Messaging/UdpMessaging/Source/UdpMessaging/Private/UdpMessagingModule.cpp
--- a/Engine/Plugins/Messaging/UdpMessaging/Source/UdpMessaging/Private/UdpMessagingModule.cpp
+++ b/Engine/Plugins/Messaging/UdpMessaging/Source/UdpMessaging/Private/UdpMessagingModule.cpp
@@ -7,6 +7,8 @@
 #include "ModuleManager.h"
 
 
+DEFINE_LOG_CATEGORY(LogUdpMessaging);
+
 #define LOCTEXT_NAMESPACE "FUdpMessagingModule"
 
 
@@ -30,7 +32,7 @@ public:
 		// load dependencies
 		if (!FModuleManager::Get().LoadModule(TEXT("Networking")).IsValid())
 		{
-			GLog->Log(TEXT("Error: The required module 'Networking' failed to load. Plug-in 'UDP Messaging' cannot be used."));
+			UE_LOG(LogUdpMessaging, Error, TEXT("The required module 'Networking' failed to load. Plug-in 'UDP Messaging' cannot be used."));
 
 			return;
 		}
@@ -100,7 +102,7 @@ protected:
 		{
 			if (!Settings->UnicastEndpoint.IsEmpty())
 			{
-				GLog->Logf(TEXT("Warning: Invalid UDP Messaging UnicastEndpoint '%s' - binding to all local network adapters instead"), *Settings->UnicastEndpoint);
+				UE_LOG(LogUdpMessaging, Warning, TEXT("Invalid UDP Messaging UnicastEndpoint '%s' - binding to all local network adapters instead"), *Settings->UnicastEndpoint);
 			}
 
 			UnicastEndpoint = FIPv4Endpoint::Any;
@@ -112,7 +114,7 @@ protected:
 		{
 			if (!Settings->MulticastEndpoint.IsEmpty())
 			{
-				GLog->Logf(TEXT("Warning: Invalid UDP Messaging MulticastEndpoint '%s' - using default endpoint '%s' instead"), *Settings->MulticastEndpoint, *UDP_MESSAGING_DEFAULT_MULTICAST_ENDPOINT.ToText().ToString());
+				UE_LOG(LogUdpMessaging, Warning, TEXT("Invalid UDP Messaging MulticastEndpoint '%s' - using default endpoint '%s' instead"), *Settings->MulticastEndpoint, *UDP_MESSAGING_DEFAULT_MULTICAST_ENDPOINT.ToText().ToString());
 			}
 
 			MulticastEndpoint = UDP_MESSAGING_DEFAULT_MULTICAST_ENDPOINT;
@@ -131,7 +133,7 @@ protected:
 			Settings->SaveConfig();
 		}
 
-		GLog->Logf(TEXT("UdpMessaging: Initializing bridge on interface %s to multicast group %s."), *UnicastEndpoint.ToText().ToString(), *MulticastEndpoint.ToText().ToString());
+		UE_LOG(LogUdpMessaging, Log, TEXT("Initializing bridge on interface %s to multicast group %s."), *UnicastEndpoint.ToText().ToString(), *MulticastEndpoint.ToText().ToString());
 
 		MessageBridge = FMessageBridgeBuilder()
 			.UsingTransport(MakeShareable(new FUdpMessageTransport(UnicastEndpoint, MulticastEndpoint, Settings->MulticastTimeToLive)));
@@ -150,7 +152,8 @@ protected:
 
 		if (!FIPv4Endpoint::Parse(Settings->TunnelUnicastEndpoint, UnicastEndpoint))
 		{
-			GLog->Logf(TEXT("Warning: Invalid UDP Tunneling UnicastEndpoint '%s' - binding to all local network adapters instead"), *Settings->UnicastEndpoint);
+			UE_LOG(LogUdpMessaging, Warning, TEXT("Invalid UDP Tunneling UnicastEndpoint '%s' - binding to all local network adapters instead"), *Settings->UnicastEndpoint);
+
 			UnicastEndpoint = FIPv4Endpoint::Any;
 			Settings->UnicastEndpoint = UnicastEndpoint.ToString();
 			ResaveSettings = true;
@@ -158,7 +161,8 @@ protected:
 
 		if (!FIPv4Endpoint::Parse(Settings->TunnelMulticastEndpoint, MulticastEndpoint))
 		{
-			GLog->Logf(TEXT("Warning: Invalid UDP Tunneling MulticastEndpoint '%s' - using default endpoint '%s' instead"), *Settings->MulticastEndpoint, *UDP_MESSAGING_DEFAULT_MULTICAST_ENDPOINT.ToText().ToString());
+			UE_LOG(LogUdpMessaging, Warning, TEXT("Invalid UDP Tunneling MulticastEndpoint '%s' - using default endpoint '%s' instead"), *Settings->MulticastEndpoint, *UDP_MESSAGING_DEFAULT_MULTICAST_ENDPOINT.ToText().ToString());
+
 			MulticastEndpoint = UDP_MESSAGING_DEFAULT_MULTICAST_ENDPOINT;
 			Settings->MulticastEndpoint = MulticastEndpoint.ToString();
 			ResaveSettings = true;
@@ -169,7 +173,7 @@ protected:
 			Settings->SaveConfig();
 		}
 
-		GLog->Logf(TEXT("UdpMessaging: Initializing tunnel on interface %s to multicast group %s."), *UnicastEndpoint.ToText().ToString(), *MulticastEndpoint.ToText().ToString());
+		UE_LOG(LogUdpMessaging, Log, TEXT("Initializing tunnel on interface %s to multicast group %s."), *UnicastEndpoint.ToText().ToString(), *MulticastEndpoint.ToText().ToString());
 
 		MessageTunnel = MakeShareable(new FUdpMessageTunnel(UnicastEndpoint, MulticastEndpoint));
 
@@ -184,7 +188,7 @@ protected:
 			}
 			else
 			{
-				GLog->Logf(TEXT("Warning: Invalid UDP RemoteTunnelEndpoint '%s' - skipping"), *Settings->RemoteTunnelEndpoints[EndpointIndex]);
+				UE_LOG(LogUdpMessaging, Warning, TEXT("Invalid UDP RemoteTunnelEndpoint '%s' - skipping"), *Settings->RemoteTunnelEndpoints[EndpointIndex]);
 			}
 		}
 	}
diff --git a/Engine/Plugins/Messaging/UdpMessaging/Source/UdpMessaging/Private/UdpMessagingPrivatePCH.h b/Engine/Plugins/Messaging/UdpMessaging/Source/UdpMessaging/Private/UdpMessagingPrivatePCH.h
--- a/Engine/Plugins/Messaging/UdpMessaging/Source/UdpMessaging/Private/UdpMessagingPrivatePCH.h
+++ b/Engine/Plugins/Messaging/UdpMessaging/Source/UdpMessaging/Private/UdpMessagingPrivatePCH.h
@@ -17,6 +17,10 @@
 /* Private constants
  *****************************************************************************/
 
+/** Declares a log category for this module. */
+DECLARE_LOG_CATEGORY_EXTERN(LogUdpMessaging, Log, All);
+
+
 /** Defines the default IP endpoint for multicast traffic. */
 #define UDP_MESSAGING_DEFAULT_MULTICAST_ENDPOINT FIPv4Endpoint(FIPv4Address(230, 0, 0, 1), 6666)
 
