UdpMessaging: Code and documentation cleanup pass

diff --git a/Engine/Plugins/Messaging/UdpMessaging/Source/UdpMessaging/Private/Transport/UdpMessageBeacon.h b/Engine/Plugins/Messaging/UdpMessaging/Source/UdpMessaging/Private/Transport/UdpMessageBeacon.h
--- a/Engine/Plugins/Messaging/UdpMessaging/Source/UdpMessaging/Private/Transport/UdpMessageBeacon.h
+++ b/Engine/Plugins/Messaging/UdpMessaging/Source/UdpMessaging/Private/Transport/UdpMessageBeacon.h
@@ -6,7 +6,7 @@
 /**
  * Implements a beacon sender thread.
  *
- * @todo gmp: add documentation for FUdpMessageBeacon
+ * @todo udpmessaging: gmp: add documentation for FUdpMessageBeacon
  */
 class FUdpMessageBeacon
 	: public FRunnable
@@ -20,9 +20,9 @@ public:
 	 * @param InSocketId The network socket identifier (used to detect unicast endpoint).
 	 * @param InMulticastEndpoint The multicast group endpoint to transport messages to.
 	 */
-	FUdpMessageBeacon( FSocket* InSocket, const FGuid& InSocketId, const FIPv4Endpoint& InMulticastEndpoint );
+	FUdpMessageBeacon(FSocket* InSocket, const FGuid& InSocketId, const FIPv4Endpoint& InMulticastEndpoint);
 
-	/** Destructor. */
+	/** Virtual destructor. */
 	virtual ~FUdpMessageBeacon();
 
 public:
@@ -42,7 +42,7 @@ public:
 	 *
 	 * @param EndpointCount The current number of known endpoints.
 	 */
-	void SetEndpointCount( int32 EndpointCount );
+	void SetEndpointCount(int32 EndpointCount);
 
 public:
 
@@ -60,7 +60,7 @@ protected:
 	 *
 	 * @param SegmentType The type of segment to send (Hello or Bye).
 	 */
-	void SendSegment( EUdpMessageSegments::Type SegmentType );
+	void SendSegment(EUdpMessageSegments::Type SegmentType);
 
 private:
 
diff --git a/Engine/Plugins/Messaging/UdpMessaging/Source/UdpMessaging/Private/Transport/UdpMessageProcessor.h b/Engine/Plugins/Messaging/UdpMessaging/Source/UdpMessaging/Private/Transport/UdpMessageProcessor.h
--- a/Engine/Plugins/Messaging/UdpMessaging/Source/UdpMessaging/Private/Transport/UdpMessageProcessor.h
+++ b/Engine/Plugins/Messaging/UdpMessaging/Source/UdpMessaging/Private/Transport/UdpMessageProcessor.h
@@ -100,10 +100,10 @@ public:
 	 * @param InNodeId The local node identifier (used to detect the unicast endpoint).
 	 * @param InMulticastEndpoint The multicast group endpoint to transport messages to.
 	 */
-	FUdpMessageProcessor( FSocket* InSocket, const FGuid& InNodeId, const FIPv4Endpoint& InMulticastEndpoint );
+	FUdpMessageProcessor(FSocket* InSocket, const FGuid& InNodeId, const FIPv4Endpoint& InMulticastEndpoint);
 
-	/** Destructor. */
-	~FUdpMessageProcessor();
+	/** Virtual destructor. */
+	virtual ~FUdpMessageProcessor();
 
 public:
 
@@ -114,7 +114,7 @@ public:
 	 * @param Sender The sender's network endpoint.
 	 * @return true if the segment was queued up, false otherwise.
 	 */
-	bool EnqueueInboundSegment( const FArrayReaderPtr& Data, const FIPv4Endpoint& Sender );
+	bool EnqueueInboundSegment(const FArrayReaderPtr& Data, const FIPv4Endpoint& Sender);
 
 	/**
 	 * Queues up an outbound message.
@@ -123,7 +123,7 @@ public:
 	 * @param Recipient The recipient's IPv4 endpoint.
 	 * @return true if the message was queued up, false otherwise.
 	 */
-	bool EnqueueOutboundMessage( const FUdpSerializedMessageRef& SerializedMessage, const FGuid& Recipient );
+	bool EnqueueOutboundMessage(const FUdpSerializedMessageRef& SerializedMessage, const FGuid& Recipient);
 
 public:
 
@@ -171,7 +171,7 @@ protected:
 	 * @param NodeInfo Details for the node to send the acknowledgment to.
 	 * @todo gmp: batch multiple of these into a single message
 	 */
-	void AcknowledgeReceipt( int32 MessageId, const FNodeInfo& NodeInfo );
+	void AcknowledgeReceipt(int32 MessageId, const FNodeInfo& NodeInfo);
 
 	/**
 	 * Calculates the time span that the thread should wait for work.
@@ -194,7 +194,7 @@ protected:
 	 * @param Sender The segment sender.
 	 * @return true if the segment passed the filter, false otherwise.
 	 */
-	bool FilterSegment( const FUdpMessageSegment::FHeader& Header, const FArrayReaderPtr& Data, const FIPv4Endpoint& Sender );
+	bool FilterSegment(const FUdpMessageSegment::FHeader& Header, const FArrayReaderPtr& Data, const FIPv4Endpoint& Sender);
 
 	/**
 	 * Processes an Abort segment.
@@ -202,7 +202,7 @@ protected:
 	 * @param Segment The segment to process.
 	 * @param NodeInfo Details for the node that sent the segment.
 	 */
-	void ProcessAbortSegment( FInboundSegment& Segment, FNodeInfo& NodeInfo );
+	void ProcessAbortSegment(FInboundSegment& Segment, FNodeInfo& NodeInfo);
 
 	/**
 	 * Processes a Success segment.
@@ -210,7 +210,7 @@ protected:
 	 * @param Segment The segment to process.
 	 * @param NodeInfo Details for the node that sent the segment.
 	 */
-	void ProcessAcknowledgeSegment( FInboundSegment& Segment, FNodeInfo& NodeInfo );
+	void ProcessAcknowledgeSegment(FInboundSegment& Segment, FNodeInfo& NodeInfo);
 
 	/**
 	 * Processes a Bye segment.
@@ -218,7 +218,7 @@ protected:
 	 * @param Segment The segment to process.
 	 * @param NodeInfo Details for the node that sent the segment.
 	 */
-	void ProcessByeSegment( FInboundSegment& Segment, FNodeInfo& NodeInfo );
+	void ProcessByeSegment(FInboundSegment& Segment, FNodeInfo& NodeInfo);
 
 	/**
 	 * Processes a Ping segment.
@@ -226,7 +226,7 @@ protected:
 	 * @param Segment The segment to process.
 	 * @param NodeInfo Details for the node that sent the segment.
 	 */
-	void ProcessDataSegment( FInboundSegment& Segment, FNodeInfo& NodeInfo );
+	void ProcessDataSegment(FInboundSegment& Segment, FNodeInfo& NodeInfo);
 
 	/**
 	 * Processes a Hello segment.
@@ -234,7 +234,7 @@ protected:
 	 * @param Segment The segment to process.
 	 * @param NodeInfo Details for the node that sent the segment.
 	 */
-	void ProcessHelloSegment( FInboundSegment& Segment, FNodeInfo& NodeInfo );
+	void ProcessHelloSegment(FInboundSegment& Segment, FNodeInfo& NodeInfo);
 
 	/**
 	 * Processes a Ping segment.
@@ -242,7 +242,7 @@ protected:
 	 * @param Segment The segment to process.
 	 * @param NodeInfo Details for the node that sent the segment.
 	 */
-	void ProcessRetransmitSegment( FInboundSegment& Segment, FNodeInfo& NodeInfo );
+	void ProcessRetransmitSegment(FInboundSegment& Segment, FNodeInfo& NodeInfo);
 
 	/**
 	 * Processes a Timeout segment.
@@ -250,7 +250,7 @@ protected:
 	 * @param Segment The segment to process.
 	 * @param NodeInfo Details for the node that sent the segment.
 	 */
-	void ProcessTimeoutSegment( FInboundSegment& Segment, FNodeInfo& NodeInfo );
+	void ProcessTimeoutSegment(FInboundSegment& Segment, FNodeInfo& NodeInfo);
 
 	/**
 	 * Processes an unknown segment type.
@@ -259,14 +259,14 @@ protected:
 	 * @param NodeInfo Details for the node that sent the segment.
 	 * @param SegmentType The segment type.
 	 */
-	void ProcessUnknownSegment( FInboundSegment& Segment, FNodeInfo& NodeInfo, uint8 SegmentType );
+	void ProcessUnknownSegment(FInboundSegment& Segment, FNodeInfo& NodeInfo, uint8 SegmentType);
 
 	/**
 	 * Removes the specified node from the list of known remote endpoints.
 	 *
 	 * @param NodeId The identifier of the node to remove.
 	 */
-	void RemoveKnownNode( const FGuid& NodeId );
+	void RemoveKnownNode(const FGuid& NodeId);
 
 	/** Updates all known remote nodes. */
 	void UpdateKnownNodes();
@@ -276,7 +276,7 @@ protected:
 	 *
 	 * @param NodeInfo Details for the node to update.
 	 */
-	void UpdateSegmenters( FNodeInfo& NodeInfo );
+	void UpdateSegmenters(FNodeInfo& NodeInfo);
 
 	/** Updates all static remote nodes. */
 	void UpdateStaticNodes();
diff --git a/Engine/Plugins/Messaging/UdpMessaging/Source/UdpMessaging/Private/Tunnel/UdpMessageTunnel.h b/Engine/Plugins/Messaging/UdpMessaging/Source/UdpMessaging/Private/Tunnel/UdpMessageTunnel.h
--- a/Engine/Plugins/Messaging/UdpMessaging/Source/UdpMessaging/Private/Tunnel/UdpMessageTunnel.h
+++ b/Engine/Plugins/Messaging/UdpMessaging/Source/UdpMessaging/Private/Tunnel/UdpMessageTunnel.h
@@ -40,8 +40,8 @@ public:
 	 */
 	FUdpMessageTunnel( const FIPv4Endpoint& InUnicastEndpoint, const FIPv4Endpoint& InMulticastEndpoint );
 
-	/** Destructor. */
-	~FUdpMessageTunnel();
+	/** Virtual destructor. */
+	virtual ~FUdpMessageTunnel();
 
 public:
 
diff --git a/Engine/Plugins/Messaging/UdpMessaging/Source/UdpMessaging/Private/Tunnel/UdpMessageTunnelConnection.h b/Engine/Plugins/Messaging/UdpMessaging/Source/UdpMessaging/Private/Tunnel/UdpMessageTunnelConnection.h
--- a/Engine/Plugins/Messaging/UdpMessaging/Source/UdpMessaging/Private/Tunnel/UdpMessageTunnelConnection.h
+++ b/Engine/Plugins/Messaging/UdpMessaging/Source/UdpMessaging/Private/Tunnel/UdpMessageTunnelConnection.h
@@ -28,8 +28,8 @@ public:
 	 */
 	FUdpMessageTunnelConnection( FSocket* InSocket, const FIPv4Endpoint& InRemoteEndpoint );
 
-	/** Destructor. */
-	~FUdpMessageTunnelConnection();
+	/** Virtual destructor. */
+	virtual ~FUdpMessageTunnelConnection();
 
 public:
 
diff --git a/Engine/Plugins/Messaging/UdpMessaging/Source/UdpMessaging/Private/UdpMessagingModule.cpp b/Engine/Plugins/Messaging/UdpMessaging/Source/UdpMessaging/Private/UdpMessagingModule.cpp
--- a/Engine/Plugins/Messaging/UdpMessaging/Source/UdpMessaging/Private/UdpMessagingModule.cpp
+++ b/Engine/Plugins/Messaging/UdpMessaging/Source/UdpMessaging/Private/UdpMessagingModule.cpp
@@ -220,7 +220,7 @@ protected:
 	/**
 	 * Checks whether networked message transport is supported.
 	 *
-	 * @todo gmp: this should be moved into an Engine module, so it can be shared with other transports
+	 * @todo udpmessaging: gmp: this should be moved into an Engine module, so it can be shared with other transports
 	 * @return true if networked transport is supported, false otherwise.
 	 */
 	bool SupportsNetworkedTransport() const
diff --git a/Engine/Source/Runtime/Networking/Public/Common/UdpSocketReceiver.h b/Engine/Source/Runtime/Networking/Public/Common/UdpSocketReceiver.h
--- a/Engine/Source/Runtime/Networking/Public/Common/UdpSocketReceiver.h
+++ b/Engine/Source/Runtime/Networking/Public/Common/UdpSocketReceiver.h
@@ -28,35 +28,32 @@ public:
 	/**
 	 * Creates and initializes a new socket receiver.
 	 *
-	 * @param InSocket - The UDP socket to receive data from.
-	 * @param InWaitTime - The amount of time to wait for the socket to be readable.
-	 * @param ThreadDescription - The thread description text (for debugging).
+	 * @param InSocket The UDP socket to receive data from.
+	 * @param InWaitTime The amount of time to wait for the socket to be readable.
+	 * @param ThreadDescription The thread description text (for debugging).
 	 */
-	FUdpSocketReceiver( FSocket* InSocket, const FTimespan& InWaitTime, const TCHAR* ThreadDescription )
+	FUdpSocketReceiver(FSocket* InSocket, const FTimespan& InWaitTime, const TCHAR* ThreadDescription)
 		: Socket(InSocket)
 		, Stopping(false)
 		, WaitTime(InWaitTime)
 	{
-		check(Socket != NULL);
+		check(Socket != nullptr);
 		check(Socket->GetSocketType() == SOCKTYPE_Datagram);
 
 		SocketSubsystem = ISocketSubsystem::Get(PLATFORM_SOCKETSUBSYSTEM);
 		Thread = FRunnableThread::Create(this, ThreadDescription, 128 * 1024, TPri_AboveNormal, FPlatformAffinity::GetPoolThreadMask());
 	}
 
-	/**
-	 * Destructor.
-	 */
-	~FUdpSocketReceiver( )
+	/** Virtual destructor. */
+	virtual ~FUdpSocketReceiver()
 	{
-		if (Thread != NULL)
+		if (Thread != nullptr)
 		{
 			Thread->Kill(true);
 			delete Thread;
 		}
 	}
 
-
 public:
 
 	/**
@@ -64,20 +61,21 @@ public:
 	 *
 	 * @return The delegate.
 	 */
-	FOnSocketDataReceived& OnDataReceived( )
+	FOnSocketDataReceived& OnDataReceived()
 	{
 		return DataReceivedDelegate;
 	}
 
-
 public:
 
-	virtual bool Init( ) override
+	// FRunnable interface
+
+	virtual bool Init() override
 	{
 		return true;
 	}
 
-	virtual uint32 Run( ) override
+	virtual uint32 Run() override
 	{
 		TSharedRef<FInternetAddr> Sender = SocketSubsystem->CreateInternetAddr();
 
@@ -105,34 +103,32 @@ public:
 		return 0;
 	}
 
-	virtual void Stop( ) override
+	virtual void Stop() override
 	{
 		Stopping = true;
 	}
 
-	virtual void Exit( ) override { }
-
+	virtual void Exit() override { }
 
 private:
 
-	// Holds the network socket.
+	/** Holds the network socket. */
 	FSocket* Socket;
 
-	// Holds a pointer to the socket sub-system.
+	/** Holds a pointer to the socket sub-system. */
 	ISocketSubsystem* SocketSubsystem;
 
-	// Holds a flag indicating that the thread is stopping.
+	/** Holds a flag indicating that the thread is stopping. */
 	bool Stopping;
 
-	// Holds the thread object.
+	/** Holds the thread object. */
 	FRunnableThread* Thread;
 
-	// Holds the amount of time to wait for inbound packets.
+	/** Holds the amount of time to wait for inbound packets. */
 	FTimespan WaitTime;
 
-
 private:
 
-	// Holds the data received delegate.
+	/** Holds the data received delegate. */
 	FOnSocketDataReceived DataReceivedDelegate;
 };
diff --git a/Engine/Source/Runtime/Networking/Public/Common/UdpSocketSender.h b/Engine/Source/Runtime/Networking/Public/Common/UdpSocketSender.h
--- a/Engine/Source/Runtime/Networking/Public/Common/UdpSocketSender.h
+++ b/Engine/Source/Runtime/Networking/Public/Common/UdpSocketSender.h
@@ -6,7 +6,7 @@
 /**
  * Asynchronously sends data to an UDP socket.
  *
- * @todo gmp: implement rate limits
+ * @todo networking: gmp: implement rate limits
  */
 class FUdpSocketSender
 	: public FRunnable
@@ -14,16 +14,16 @@ class FUdpSocketSender
 	// Structure for outbound packets.
 	struct FPacket
 	{
-		// Holds the packet's data.
+		/** Holds the packet's data. */
 		TSharedPtr<TArray<uint8>, ESPMode::ThreadSafe> Data;
 
-		// Holds the recipient.
+		/** Holds the recipient. */
 		FIPv4Endpoint Recipient;
 
-		// Default constructor.
+		/** Default constructor. */
 		FPacket( ) { }
 
-		// Creates and initializes a new instance.
+		/** Creates and initializes a new instance. */
 		FPacket( const TSharedRef<TArray<uint8>, ESPMode::ThreadSafe>& InData, const FIPv4Endpoint& InRecipient )
 			: Data(InData)
 			, Recipient(InRecipient)
@@ -35,8 +35,8 @@ public:
 	/**
 	 * Creates and initializes a new socket sender.
 	 *
-	 * @param InSocket - The UDP socket to use for sending data.
-	 * @param ThreadDescription - The thread description text (for debugging).
+	 * @param InSocket The UDP socket to use for sending data.
+	 * @param ThreadDescription The thread description text (for debugging).
 	 */
 	FUdpSocketSender( FSocket* InSocket, const TCHAR* ThreadDescription )
 		: SendRate(0)
@@ -44,7 +44,7 @@ public:
 		, Stopping(false)
 		, WaitTime(FTimespan::FromMilliseconds(100))
 	{
-		check(Socket != NULL);
+		check(Socket != nullptr);
 		check(Socket->GetSocketType() == SOCKTYPE_Datagram);
 
 		int32 NewSize = 0;
@@ -55,12 +55,10 @@ public:
 		Thread = FRunnableThread::Create(this, ThreadDescription, 128 * 1024, TPri_AboveNormal, FPlatformAffinity::GetPoolThreadMask());
 	}
 
-	/**
-	 * Destructor.
-	 */
-	~FUdpSocketSender()
+	/** Virtual destructor. */
+	virtual ~FUdpSocketSender()
 	{
-		if (Thread != NULL)
+		if (Thread != nullptr)
 		{
 			Thread->Kill(true);
 			delete Thread;
@@ -95,9 +93,8 @@ public:
 	/**
 	 * Sends data to the specified recipient.
 	 *
-	 * @param Data - The data to send.
-	 * @param Recipient - The recipient.
-	 *
+	 * @param Data The data to send.
+	 * @param Recipient The recipient.
 	 * @return true if the data will be sent, false otherwise.
 	 */
 	bool Send( const TSharedRef<TArray<uint8>, ESPMode::ThreadSafe>& Data, const FIPv4Endpoint& Recipient )
@@ -115,7 +112,8 @@ public:
 	/**
 	 * Sets the send rate (in bytes per second).
 	 *
-	 * @param Rate - The new send rate (0 = unlimited).
+	 * @param Rate The new send rate (0 = unlimited).
+	 * @see SetWaitTime
 	 */
 	void SetSendRate( uint32 Rate )
 	{
@@ -125,7 +123,8 @@ public:
 	/**
 	 * Sets the maximum time span to wait for work items.
 	 *
-	 * @param Timespan - The wait time.
+	 * @param Timespan The wait time.
+	 * @see SetSendRate
 	 */
 	void SetWaitTime( const FTimespan& Timespan )
 	{
@@ -134,12 +133,14 @@ public:
 
 public:
 
-	virtual bool Init( ) override
+	// FRunnable interface
+
+	virtual bool Init() override
 	{
 		return true;
 	}
 
-	virtual uint32 Run( ) override
+	virtual uint32 Run() override
 	{
 		while (!Stopping)
 		{
@@ -168,39 +169,38 @@ public:
 		return 0;
 	}
 
-	virtual void Stop( ) override
+	virtual void Stop() override
 	{
 		Stopping = true;
 
 		WorkEvent->Trigger();
 	}
 
-	virtual void Exit( ) override { }
+	virtual void Exit() override { }
 
 private:
 
-	// Holds the send queue.
+	/** Holds the send queue. */
 	TQueue<FPacket, EQueueMode::Mpsc> SendQueue;
 
-	// Holds the send rate.
+	/** Holds the send rate. */
 	uint32 SendRate;
 
-	// Holds the network socket.
+	/** Holds the network socket. */
 	FSocket* Socket;
 
-	// Holds a flag indicating that the thread is stopping.
+	/** Holds a flag indicating that the thread is stopping. */
 	bool Stopping;
 
-	// Holds the thread object.
+	/** Holds the thread object. */
 	FRunnableThread* Thread;
 
-	// Holds the current throughput.
+	/** Holds the current throughput. */
 	uint32 Throughput;
 
-	
-	// Holds the amount of time to wait for inbound packets.
+	/** Holds the amount of time to wait for inbound packets. */
 	FTimespan WaitTime;
 
-	// Holds an event signaling that inbound messages need to be processed.
+	/** Holds an event signaling that inbound messages need to be processed. */
 	FEvent* WorkEvent;
 };
