UdpMessaging: More helpful settings tooltips.

diff --git a/Engine/Plugins/Messaging/UdpMessaging/Source/UdpMessaging/Private/Shared/UdpMessagingSettings.h b/Engine/Plugins/Messaging/UdpMessaging/Source/UdpMessaging/Private/Shared/UdpMessagingSettings.h
--- a/Engine/Plugins/Messaging/UdpMessaging/Source/UdpMessaging/Private/Shared/UdpMessagingSettings.h
+++ b/Engine/Plugins/Messaging/UdpMessaging/Source/UdpMessaging/Private/Shared/UdpMessagingSettings.h
@@ -17,22 +17,34 @@ public:
 	UPROPERTY(config, EditAnywhere, Category=Transport)
 	bool EnableTransport;
 
-	/** The IP endpoint to listen to and send packets from. */
+	/**
+	 * The IP endpoint to listen to and send packets from.
+	 *
+	 * The format is IP_ADDRESS:PORT_NUMBER.
+	 * 0.0.0.0:0 will bind to the default network adapter on Windows
+	 * and all available network adapters on other operating systems.
+	 */
 	UPROPERTY(config, EditAnywhere, Category=Transport)
 	FString UnicastEndpoint;
 
-	/** The IP endpoint to send multicast packets to. */
+	/**
+	 * The IP endpoint to send multicast packets to.
+	 *
+	 * The format is IP_ADDRESS:PORT_NUMBER.
+	 * The multicast IP address must be in the range 224.0.0.0 to 239.255.255.255.
+	 */
 	UPROPERTY(config, EditAnywhere, Category=Transport)
 	FString MulticastEndpoint;
 
 	/** The time-to-live (TTL) for sent multicast packets. */
-	UPROPERTY(config, EditAnywhere, Category=Transport)
+	UPROPERTY(config, EditAnywhere, Category=Transport, AdvancedDisplay)
 	uint8 MulticastTimeToLive;
 
 	/**
 	 * The IP endpoints of static devices.
 	 *
 	 * Use this setting to list devices on other subnets, such as mobile phones on a WiFi network.
+	 * The format is IP_ADDRESS:PORT_NUMBER.
 	 */
 	UPROPERTY(config, EditAnywhere, Category=Transport, AdvancedDisplay)
 	TArray<FString> StaticEndpoints;
@@ -43,11 +55,20 @@ public:
 	UPROPERTY(config, EditAnywhere, Category=Tunnel)
 	bool EnableTunnel;
 
-	/** The IP endpoint to listen to and send packets from. */
+	/**
+	 * The local IP endpoint to listen to and send packets from.
+	 *
+	 * The format is IP_ADDRESS:PORT_NUMBER.
+	 */
 	UPROPERTY(config, EditAnywhere, Category=Tunnel)
 	FString TunnelUnicastEndpoint;
 
-	/** The IP endpoint to send multicast packets to. */
+	/**
+	 * The IP endpoint to send multicast packets to.
+	 *
+	 * The format is IP_ADDRESS:PORT_NUMBER.
+	 * The multicast IP address must be in the range 224.0.0.0 to 239.255.255.255.
+	 */
 	UPROPERTY(config, EditAnywhere, Category=Tunnel)
 	FString TunnelMulticastEndpoint;
 
@@ -55,6 +76,7 @@ public:
 	 * The IP endpoints of remote tunnel nodes.
 	 *
 	 * Use this setting to connect to remote tunnel services.
+	 * The format is IP_ADDRESS:PORT_NUMBER.
 	 */
 	UPROPERTY(config, EditAnywhere, Category=Tunnel, AdvancedDisplay)
 	TArray<FString> RemoteTunnelEndpoints;
