UdpMessaging: Fixed a race condition during message serialization

diff --git a/Engine/Plugins/Messaging/UdpMessaging/Source/UdpMessaging/Private/Transport/UdpMessageProcessor.cpp b/Engine/Plugins/Messaging/UdpMessaging/Source/UdpMessaging/Private/Transport/UdpMessageProcessor.cpp
--- a/Engine/Plugins/Messaging/UdpMessaging/Source/UdpMessaging/Private/Transport/UdpMessageProcessor.cpp
+++ b/Engine/Plugins/Messaging/UdpMessaging/Source/UdpMessaging/Private/Transport/UdpMessageProcessor.cpp
@@ -78,11 +78,6 @@ bool FUdpMessageProcessor::EnqueueOutboundMessage( const FUdpSerializedMessageRe
 
 	SerializedMessage->OnStateChanged().BindRaw(this, &FUdpMessageProcessor::HandleSerializedMessageStateChanged);
 
-	if (SerializedMessage->GetState() != EUdpSerializedMessageState::Incomplete)
-	{
-		WorkEvent->Trigger();
-	}
-
 	return true;
 }
 
@@ -128,7 +123,6 @@ uint32 FUdpMessageProcessor::Run()
 void FUdpMessageProcessor::Stop()
 {
 	Stopping = true;
-
 	WorkEvent->Trigger();
 }
 
diff --git a/Engine/Plugins/Messaging/UdpMessaging/Source/UdpMessaging/Private/Transport/UdpMessageTransport.cpp b/Engine/Plugins/Messaging/UdpMessaging/Source/UdpMessaging/Private/Transport/UdpMessageTransport.cpp
--- a/Engine/Plugins/Messaging/UdpMessaging/Source/UdpMessaging/Private/Transport/UdpMessageTransport.cpp
+++ b/Engine/Plugins/Messaging/UdpMessaging/Source/UdpMessaging/Private/Transport/UdpMessageTransport.cpp
@@ -123,23 +123,23 @@ bool FUdpMessageTransport::TransportMessage( const IMessageContextRef& Context,
 	}
 
 	FUdpSerializedMessageRef SerializedMessage = MakeShareable(new FUdpSerializedMessage());
-	TGraphTask<FUdpSerializeMessageTask>::CreateTask().ConstructAndDispatchWhenReady(Context, SerializedMessage);
 
-	// publish the message
 	if (Recipients.Num() == 0)
 	{
-		return MessageProcessor->EnqueueOutboundMessage(SerializedMessage, FGuid());
+		// publish the message
+		MessageProcessor->EnqueueOutboundMessage(SerializedMessage, FGuid());
 	}
-
-	// send the message
-	for (int32 Index = 0; Index < Recipients.Num(); ++Index)
+	else
 	{
-		if (!MessageProcessor->EnqueueOutboundMessage(SerializedMessage, Recipients[Index]))
+		// send the message
+		for (const auto& Recipient : Recipients)
 		{
-			return false;
+			MessageProcessor->EnqueueOutboundMessage(SerializedMessage, Recipient);
 		}
 	}
 
+	TGraphTask<FUdpSerializeMessageTask>::CreateTask().ConstructAndDispatchWhenReady(Context, SerializedMessage);
+
 	return true;
 }
 
