UdpMessaging: Modernization pass (ranged for-loops, reduced nesting, fixed documentation)

diff --git a/Engine/Plugins/Messaging/UdpMessaging/Source/UdpMessaging/Private/Transport/UdpMessageBeacon.cpp b/Engine/Plugins/Messaging/UdpMessaging/Source/UdpMessaging/Private/Transport/UdpMessageBeacon.cpp
--- a/Engine/Plugins/Messaging/UdpMessaging/Source/UdpMessaging/Private/Transport/UdpMessageBeacon.cpp
+++ b/Engine/Plugins/Messaging/UdpMessaging/Source/UdpMessaging/Private/Transport/UdpMessageBeacon.cpp
@@ -109,15 +109,17 @@ void FUdpMessageBeacon::Stop()
 void FUdpMessageBeacon::SendSegment(EUdpMessageSegments::Type SegmentType)
 {
 	FUdpMessageSegment::FHeader Header;
-
-	Header.SenderNodeId = NodeId;
-	Header.ProtocolVersion = UDP_MESSAGING_TRANSPORT_PROTOCOL_VERSION;
-	Header.SegmentType = SegmentType;
+	{
+		Header.SenderNodeId = NodeId;
+		Header.ProtocolVersion = UDP_MESSAGING_TRANSPORT_PROTOCOL_VERSION;
+		Header.SegmentType = SegmentType;
+	}
 
 	FArrayWriter Writer;
-	
-	Writer << Header;
-	Writer << NodeId;
+	{
+		Writer << Header;
+		Writer << NodeId;
+	}
 
 	int32 Sent;
 
diff --git a/Engine/Plugins/Messaging/UdpMessaging/Source/UdpMessaging/Private/Transport/UdpMessageProcessor.cpp b/Engine/Plugins/Messaging/UdpMessaging/Source/UdpMessaging/Private/Transport/UdpMessageProcessor.cpp
--- a/Engine/Plugins/Messaging/UdpMessaging/Source/UdpMessaging/Private/Transport/UdpMessageProcessor.cpp
+++ b/Engine/Plugins/Messaging/UdpMessaging/Source/UdpMessaging/Private/Transport/UdpMessageProcessor.cpp
@@ -12,7 +12,7 @@ const int32 FUdpMessageProcessor::DeadHelloIntervals = 5;
 /* FUdpMessageProcessor structors
  *****************************************************************************/
 
-FUdpMessageProcessor::FUdpMessageProcessor( FSocket* InSocket, const FGuid& InNodeId, const FIPv4Endpoint& InMulticastEndpoint )
+FUdpMessageProcessor::FUdpMessageProcessor(FSocket* InSocket, const FGuid& InNodeId, const FIPv4Endpoint& InMulticastEndpoint)
 	: Beacon(nullptr)
 	, LastSentMessage(-1)
 	, LocalNodeId(InNodeId)
@@ -56,7 +56,7 @@ FUdpMessageProcessor::~FUdpMessageProcessor()
 /* FUdpMessageProcessor interface
  *****************************************************************************/
 
-bool FUdpMessageProcessor::EnqueueInboundSegment( const FArrayReaderPtr& Data, const FIPv4Endpoint& InSender ) 
+bool FUdpMessageProcessor::EnqueueInboundSegment(const FArrayReaderPtr& Data, const FIPv4Endpoint& InSender) 
 {
 	if (!InboundSegments.Enqueue(FInboundSegment(Data, InSender)))
 	{
@@ -69,7 +69,7 @@ bool FUdpMessageProcessor::EnqueueInboundSegment( const FArrayReaderPtr& Data, c
 }
 
 
-bool FUdpMessageProcessor::EnqueueOutboundMessage( const FUdpSerializedMessageRef& SerializedMessage, const FGuid& Recipient )
+bool FUdpMessageProcessor::EnqueueOutboundMessage(const FUdpSerializedMessageRef& SerializedMessage, const FGuid& Recipient)
 {
 	if (!OutboundMessages.Enqueue(FOutboundMessage(SerializedMessage, Recipient)))
 	{
@@ -130,23 +130,26 @@ void FUdpMessageProcessor::Stop()
 /* FUdpMessageProcessor implementation
  *****************************************************************************/
 
-void FUdpMessageProcessor::AcknowledgeReceipt( int32 MessageId, const FNodeInfo& NodeInfo )
+void FUdpMessageProcessor::AcknowledgeReceipt(int32 MessageId, const FNodeInfo& NodeInfo)
 {
 	FUdpMessageSegment::FHeader Header;
-
-	Header.RecipientNodeId = NodeInfo.NodeId;
-	Header.SenderNodeId = LocalNodeId;
-	Header.ProtocolVersion = UDP_MESSAGING_TRANSPORT_PROTOCOL_VERSION;
-	Header.SegmentType = EUdpMessageSegments::Acknowledge;
+	{
+		Header.RecipientNodeId = NodeInfo.NodeId;
+		Header.SenderNodeId = LocalNodeId;
+		Header.ProtocolVersion = UDP_MESSAGING_TRANSPORT_PROTOCOL_VERSION;
+		Header.SegmentType = EUdpMessageSegments::Acknowledge;
+	}
 
 	FUdpMessageSegment::FAcknowledgeChunk AcknowledgeChunk;
-
-	AcknowledgeChunk.MessageId = MessageId;
+	{
+		AcknowledgeChunk.MessageId = MessageId;
+	}
 
 	FArrayWriter Writer;
-
-	Writer << Header;
-	Writer << AcknowledgeChunk;
+	{
+		Writer << Header;
+		Writer << AcknowledgeChunk;
+	}
 
 	int32 Sent;
 
@@ -237,19 +240,24 @@ void FUdpMessageProcessor::ConsumeOutboundMessages()
 		{
 			RecipientNodeInfo.Endpoint = MulticastEndpoint;
 
-			for (TMap<FIPv4Endpoint, FNodeInfo>::TIterator It(StaticNodes); It; ++It)
+			for (auto& StaticNodeInfoPair : StaticNodes)
 			{
-				FNodeInfo& StaticNodeInfo = It.Value();
-				StaticNodeInfo.Segmenters.Add(LastSentMessage, MakeShareable(new FUdpMessageSegmenter(OutboundMessage.SerializedMessage.ToSharedRef(), 1024)));
+				StaticNodeInfoPair.Value.Segmenters.Add(
+					LastSentMessage,
+					MakeShareable(new FUdpMessageSegmenter(OutboundMessage.SerializedMessage.ToSharedRef(), 1024))
+				);
 			}
 		}
 
-		RecipientNodeInfo.Segmenters.Add(LastSentMessage, MakeShareable(new FUdpMessageSegmenter(OutboundMessage.SerializedMessage.ToSharedRef(), 1024)));
+		RecipientNodeInfo.Segmenters.Add(
+			LastSentMessage,
+			MakeShareable(new FUdpMessageSegmenter(OutboundMessage.SerializedMessage.ToSharedRef(), 1024))
+		);
 	}
 }
 
 
-bool FUdpMessageProcessor::FilterSegment( const FUdpMessageSegment::FHeader& Header, const FArrayReaderPtr& Data, const FIPv4Endpoint& InSender )
+bool FUdpMessageProcessor::FilterSegment(const FUdpMessageSegment::FHeader& Header, const FArrayReaderPtr& Data, const FIPv4Endpoint& InSender)
 {
 	// filter unsupported protocol versions
 	if (Header.ProtocolVersion != UDP_MESSAGING_TRANSPORT_PROTOCOL_VERSION)
@@ -267,43 +275,38 @@ bool FUdpMessageProcessor::FilterSegment( const FUdpMessageSegment::FHeader& Hea
 }
 
 
-void FUdpMessageProcessor::ProcessAbortSegment( FInboundSegment& Segment, FNodeInfo& NodeInfo )
+void FUdpMessageProcessor::ProcessAbortSegment(FInboundSegment& Segment, FNodeInfo& NodeInfo)
 {
 	FUdpMessageSegment::FAbortChunk AbortChunk;
 
 	*Segment.Data << AbortChunk;
-
 	NodeInfo.Segmenters.Remove(AbortChunk.MessageId);
 }
 
 
-void FUdpMessageProcessor::ProcessAcknowledgeSegment( FInboundSegment& Segment, FNodeInfo& NodeInfo )
+void FUdpMessageProcessor::ProcessAcknowledgeSegment(FInboundSegment& Segment, FNodeInfo& NodeInfo)
 {
 	FUdpMessageSegment::FAcknowledgeChunk AcknowledgeChunk;
 
 	*Segment.Data << AcknowledgeChunk;
-
 	NodeInfo.Segmenters.Remove(AcknowledgeChunk.MessageId);
 }
 
 
-void FUdpMessageProcessor::ProcessByeSegment( FInboundSegment& Segment, FNodeInfo& NodeInfo )
+void FUdpMessageProcessor::ProcessByeSegment(FInboundSegment& Segment, FNodeInfo& NodeInfo)
 {
 	FGuid RemoteNodeId;
 
 	*Segment.Data << RemoteNodeId;
 
-	if (RemoteNodeId.IsValid())
+	if (RemoteNodeId.IsValid() && (RemoteNodeId == NodeInfo.NodeId))
 	{
-		if (NodeInfo.NodeId == RemoteNodeId)
-		{
-			RemoveKnownNode(RemoteNodeId);
-		}
+		RemoveKnownNode(RemoteNodeId);
 	}
 }
 
 
-void FUdpMessageProcessor::ProcessDataSegment( FInboundSegment& Segment, FNodeInfo& NodeInfo )
+void FUdpMessageProcessor::ProcessDataSegment(FInboundSegment& Segment, FNodeInfo& NodeInfo)
 {
 	FUdpMessageSegment::FDataChunk DataChunk;
 
@@ -326,36 +329,38 @@ void FUdpMessageProcessor::ProcessDataSegment( FInboundSegment& Segment, FNodeIn
 	ReassembledMessage->Reassemble(DataChunk.SegmentNumber, DataChunk.SegmentOffset, DataChunk.Data, CurrentTime);
 
 	// Deliver or re-sequence message
-	if (ReassembledMessage->IsComplete())
+	if (!ReassembledMessage->IsComplete())
 	{
-		AcknowledgeReceipt(DataChunk.MessageId, NodeInfo);
+		return;
+	}
+
+	AcknowledgeReceipt(DataChunk.MessageId, NodeInfo);
 
-		if (ReassembledMessage->GetSequence() == 0)
+	if (ReassembledMessage->GetSequence() == 0)
+	{
+		if (NodeInfo.NodeId.IsValid())
 		{
-			if (NodeInfo.NodeId.IsValid())
-			{
-				MessageReassembledDelegate.ExecuteIfBound(ReassembledMessage.ToSharedRef(), nullptr, NodeInfo.NodeId);
-			}
+			MessageReassembledDelegate.ExecuteIfBound(ReassembledMessage.ToSharedRef(), nullptr, NodeInfo.NodeId);
 		}
-		else if (NodeInfo.Resequencer.Resequence(ReassembledMessage))
-		{
-			FUdpReassembledMessagePtr ResequencedMessage;
+	}
+	else if (NodeInfo.Resequencer.Resequence(ReassembledMessage))
+	{
+		FUdpReassembledMessagePtr ResequencedMessage;
 
-			while (NodeInfo.Resequencer.Pop(ResequencedMessage))
+		while (NodeInfo.Resequencer.Pop(ResequencedMessage))
+		{
+			if (NodeInfo.NodeId.IsValid())
 			{
-				if (NodeInfo.NodeId.IsValid())
-				{
-					MessageReassembledDelegate.ExecuteIfBound(ResequencedMessage.ToSharedRef(), nullptr, NodeInfo.NodeId);
-				}
+				MessageReassembledDelegate.ExecuteIfBound(ResequencedMessage.ToSharedRef(), nullptr, NodeInfo.NodeId);
 			}
 		}
-
-		NodeInfo.ReassembledMessages.Remove(DataChunk.MessageId);
 	}
+
+	NodeInfo.ReassembledMessages.Remove(DataChunk.MessageId);
 }
 
 
-void FUdpMessageProcessor::ProcessHelloSegment( FInboundSegment& Segment, FNodeInfo& NodeInfo )
+void FUdpMessageProcessor::ProcessHelloSegment(FInboundSegment& Segment, FNodeInfo& NodeInfo)
 {
 	FGuid RemoteNodeId;
 
@@ -368,7 +373,7 @@ void FUdpMessageProcessor::ProcessHelloSegment( FInboundSegment& Segment, FNodeI
 }
 
 
-void FUdpMessageProcessor::ProcessRetransmitSegment( FInboundSegment& Segment, FNodeInfo& NodeInfo )
+void FUdpMessageProcessor::ProcessRetransmitSegment(FInboundSegment& Segment, FNodeInfo& NodeInfo)
 {
 	FUdpMessageSegment::FRetransmitChunk RetransmitChunk;
 
@@ -383,7 +388,7 @@ void FUdpMessageProcessor::ProcessRetransmitSegment( FInboundSegment& Segment, F
 }
 
 
-void FUdpMessageProcessor::ProcessTimeoutSegment( FInboundSegment& Segment, FNodeInfo& NodeInfo )
+void FUdpMessageProcessor::ProcessTimeoutSegment(FInboundSegment& Segment, FNodeInfo& NodeInfo)
 {
 	FUdpMessageSegment::FTimeoutChunk TimeoutChunk;
 
@@ -398,13 +403,13 @@ void FUdpMessageProcessor::ProcessTimeoutSegment( FInboundSegment& Segment, FNod
 }
 
 
-void FUdpMessageProcessor::ProcessUnknownSegment( FInboundSegment& Segment, FNodeInfo& EndpointInfo, uint8 SegmentType )
+void FUdpMessageProcessor::ProcessUnknownSegment(FInboundSegment& Segment, FNodeInfo& EndpointInfo, uint8 SegmentType)
 {
 	UE_LOG(LogUdpMessaging, Verbose, TEXT("Received unknown segment type '%i' from %s"), SegmentType, *Segment.Sender.ToText().ToString());
 }
 
 
-void FUdpMessageProcessor::RemoveKnownNode( const FGuid& NodeId )
+void FUdpMessageProcessor::RemoveKnownNode(const FGuid& NodeId)
 {
 	NodeLostDelegate.ExecuteIfBound(NodeId);
 
@@ -419,13 +424,14 @@ void FUdpMessageProcessor::UpdateKnownNodes()
 
 	TArray<FGuid> NodesToRemove;
 
-	for (TMap<FGuid, FNodeInfo>::TIterator It(KnownNodes); It; ++It)
+	for (auto& KnownNodePair : KnownNodes)
 	{
-		FNodeInfo& NodeInfo = It.Value();
+		FGuid& NodeId = KnownNodePair.Key;
+		FNodeInfo& NodeInfo = KnownNodePair.Value;
 
-		if ((It.Key().IsValid()) && ((NodeInfo.LastSegmentReceivedTime + DeadHelloTimespan) <= CurrentTime))
+		if ((NodeId.IsValid()) && ((NodeInfo.LastSegmentReceivedTime + DeadHelloTimespan) <= CurrentTime))
 		{
-			NodesToRemove.Add(It.Key());
+			NodesToRemove.Add(NodeId);
 		}
 		else
 		{
@@ -433,28 +439,29 @@ void FUdpMessageProcessor::UpdateKnownNodes()
 		}
 	}
 
-	for (int32 Index = 0; Index < NodesToRemove.Num(); ++Index)
+	for (const auto& Node : NodesToRemove)
 	{
-		// @todo gmp: put this back in after testing
-		//RemoveKnownEndpoint(EndpointsToRemove(Index));
+		// @todo udpmessaging: gmp: put this back in after testing
+		//RemoveKnownNode(Node);
 	}
 
 	Beacon->SetEndpointCount(KnownNodes.Num() + 1);
 }
 
 
-void FUdpMessageProcessor::UpdateSegmenters( FNodeInfo& NodeInfo )
+void FUdpMessageProcessor::UpdateSegmenters(FNodeInfo& NodeInfo)
 {
 	FUdpMessageSegment::FHeader Header;
+	{
+		Header.RecipientNodeId = NodeInfo.NodeId;
+		Header.SenderNodeId = LocalNodeId;
+		Header.ProtocolVersion = UDP_MESSAGING_TRANSPORT_PROTOCOL_VERSION;
+		Header.SegmentType = EUdpMessageSegments::Data;
+	}
 
-	Header.RecipientNodeId = NodeInfo.NodeId;
-	Header.SenderNodeId = LocalNodeId;
-	Header.ProtocolVersion = UDP_MESSAGING_TRANSPORT_PROTOCOL_VERSION;
-	Header.SegmentType = EUdpMessageSegments::Data;
-
-	for (TMap<int32, TSharedPtr<FUdpMessageSegmenter> >::TIterator It2(NodeInfo.Segmenters); It2; ++It2)
+	for (TMap<int32, TSharedPtr<FUdpMessageSegmenter> >::TIterator It(NodeInfo.Segmenters); It; ++It)
 	{
-		TSharedPtr<FUdpMessageSegmenter>& Segmenter = It2.Value();
+		TSharedPtr<FUdpMessageSegmenter>& Segmenter = It.Value();
 
 		Segmenter->Initialize();
 
@@ -464,7 +471,7 @@ void FUdpMessageProcessor::UpdateSegmenters( FNodeInfo& NodeInfo )
 
 			while (Segmenter->GetNextPendingSegment(DataChunk.Data, DataChunk.SegmentNumber))
 			{
-				DataChunk.MessageId = It2.Key();
+				DataChunk.MessageId = It.Key();
 				DataChunk.MessageSize = Segmenter->GetMessageSize();
 				DataChunk.SegmentOffset = 1024 * DataChunk.SegmentNumber;
 				DataChunk.Sequence = 0;
@@ -485,11 +492,11 @@ void FUdpMessageProcessor::UpdateSegmenters( FNodeInfo& NodeInfo )
 				}
 			}
 
-			It2.RemoveCurrent();
+			It.RemoveCurrent();
 		}
 		else if (Segmenter->IsInvalid())
 		{
-			It2.RemoveCurrent();
+			It.RemoveCurrent();
 		}
 	}
 }
@@ -497,9 +504,9 @@ void FUdpMessageProcessor::UpdateSegmenters( FNodeInfo& NodeInfo )
 
 void FUdpMessageProcessor::UpdateStaticNodes()
 {
-	for (TMap<FIPv4Endpoint, FNodeInfo>::TIterator It(StaticNodes); It; ++It)
+	for (auto& StaticNodePair : StaticNodes)
 	{
-		UpdateSegmenters(It.Value());
+		UpdateSegmenters(StaticNodePair.Value);
 	}
 }
 
diff --git a/Engine/Plugins/Messaging/UdpMessaging/Source/UdpMessaging/Private/Transport/UdpMessageProcessor.h b/Engine/Plugins/Messaging/UdpMessaging/Source/UdpMessaging/Private/Transport/UdpMessageProcessor.h
--- a/Engine/Plugins/Messaging/UdpMessaging/Source/UdpMessaging/Private/Transport/UdpMessageProcessor.h
+++ b/Engine/Plugins/Messaging/UdpMessaging/Source/UdpMessaging/Private/Transport/UdpMessageProcessor.h
@@ -9,35 +9,35 @@
 class FUdpMessageProcessor
 	: public FRunnable
 {
-	// Structure for known remote endpoints.
+	/** Structure for known remote endpoints. */
 	struct FNodeInfo
 	{
-		// Holds the node's IP endpoint.
+		/** Holds the node's IP endpoint. */
 		FIPv4Endpoint Endpoint;
 
-		// Holds the time at which the last Hello was received.
+		/** Holds the time at which the last Hello was received. */
 		FDateTime LastSegmentReceivedTime;
 
-		// Holds the endpoint's node identifier.
+		/** Holds the endpoint's node identifier. */
 		FGuid NodeId;
 
-		// Holds the collection of reassembled messages.
+		/** Holds the collection of reassembled messages. */
 		TMap<int32, FUdpReassembledMessagePtr> ReassembledMessages;
 
-		// Holds the message resequencer.
+		/** Holds the message resequencer. */
 		FUdpMessageResequencer Resequencer;
 
-		// Holds the collection of message segmenters.
+		/** Holds the collection of message segmenters. */
 		TMap<int32, TSharedPtr<FUdpMessageSegmenter> > Segmenters;
 
-		// Default constructor.
+		/** Default constructor. */
 		FNodeInfo()
 			: LastSegmentReceivedTime(FDateTime::MinValue())
 			, NodeId()
 		{ }
 
-		// Resets the endpoint info.
-		void ResetIfRestarted( const FGuid& NewNodeId )
+		/** Resets the endpoint info. */
+		void ResetIfRestarted(const FGuid& NewNodeId)
 		{
 			if (NewNodeId != NodeId)
 			{
@@ -50,21 +50,20 @@ class FUdpMessageProcessor
 	};
 
 
-	// Structure for inbound segments.
+	/** Structure for inbound segments. */
 	struct FInboundSegment
 	{
-		// Holds the segment data.
+		/** Holds the segment data. */
 		FArrayReaderPtr Data;
 
-		// Holds the sender's network endpoint.
+		/** Holds the sender's network endpoint. */
 		FIPv4Endpoint Sender;
 
-
-		// Default constructor.
+		/** Default constructor. */
 		FInboundSegment() { }
 
-		// Creates and initializes a new instance.
-		FInboundSegment( const FArrayReaderPtr& InData, const FIPv4Endpoint& InSender )
+		/** Creates and initializes a new instance. */
+		FInboundSegment(const FArrayReaderPtr& InData, const FIPv4Endpoint& InSender)
 			: Data(InData)
 			, Sender(InSender)
 		{ }
@@ -74,18 +73,17 @@ class FUdpMessageProcessor
 	// Structure for outbound messages.
 	struct FOutboundMessage
 	{
-		// Holds the serialized message.
+		/** Holds the serialized message. */
 		FUdpSerializedMessagePtr SerializedMessage;
 
-		// Holds the recipient.
+		/** Holds the recipient. */
 		FGuid RecipientId;
 
-
-		// Default constructor.
+		/** Default constructor. */
 		FOutboundMessage() { }
 
-		// Creates and initializes a new instance.
-		FOutboundMessage( const FUdpSerializedMessageRef& InSerializedMessage, const FGuid& InRecipientId )
+		/** Creates and initializes a new instance. */
+		FOutboundMessage(const FUdpSerializedMessageRef& InSerializedMessage, const FGuid& InRecipientId)
 			: SerializedMessage(InSerializedMessage)
 			, RecipientId(InRecipientId)
 		{ }
@@ -130,7 +128,7 @@ public:
 	/**
 	 * Returns a delegate that is executed when message data has been received.
 	 *
-	 * @param Delegate The delegate to set.
+	 * @return The delegate.
 	 */
 	DECLARE_DELEGATE_ThreeParams(FOnMessageReassembled, const FUdpReassembledMessageRef& /*ReassembledMessage*/, const IMessageAttachmentPtr& /*Attachment*/, const FGuid& /*NodeId*/)
 	FOnMessageReassembled& OnMessageReassembled()
@@ -138,15 +136,22 @@ public:
 		return MessageReassembledDelegate;
 	}
 
+	/**
+	 * Returns a delegate that is executed when a remote node has been discovered.
+	 *
+	 * @return The delegate.
+	 * @see OnNodeLost
+	 */
 	IMessageTransport::FOnNodeDiscovered& OnNodeDiscovered()
 	{
 		return NodeDiscoveredDelegate;
 	}
 
 	/**
-	 * Returns a delegate that is executed when a channel has closed or timed out.
+	 * Returns a delegate that is executed when a remote node was closed or timed out.
 	 *
-	 * @param Delegate The delegate to set.
+	 * @return The delegate.
+	 * @see OnNodeDiscovered
 	 */
 	IMessageTransport::FOnNodeLost& OnNodeLost()
 	{
@@ -288,10 +293,10 @@ private:
 
 private:
 
-	// Holds the queue of outbound messages. */
+	/** Holds the queue of outbound messages. */
 	TQueue<FInboundSegment, EQueueMode::Mpsc> InboundSegments;
 
-	// Holds the queue of outbound messages. */
+	/** Holds the queue of outbound messages. */
 	TQueue<FOutboundMessage, EQueueMode::Mpsc> OutboundMessages;
 
 private:
diff --git a/Engine/Plugins/Messaging/UdpMessaging/Source/UdpMessaging/Private/Transport/UdpMessageResequencer.h b/Engine/Plugins/Messaging/UdpMessaging/Source/UdpMessaging/Private/Transport/UdpMessageResequencer.h
--- a/Engine/Plugins/Messaging/UdpMessaging/Source/UdpMessaging/Private/Transport/UdpMessageResequencer.h
+++ b/Engine/Plugins/Messaging/UdpMessaging/Source/UdpMessaging/Private/Transport/UdpMessageResequencer.h
@@ -40,16 +40,15 @@ public:
 	 */
 	bool Pop(FUdpReassembledMessagePtr& OutMessage)
 	{
-		if (MessageHeap.HeapTop()->GetSequence() == NextSequence)
+		if (MessageHeap.HeapTop()->GetSequence() != NextSequence)
 		{
-			MessageHeap.HeapPop(OutMessage, FSequenceComparer());
-
-			NextSequence++;
-
-			return true;
+			return false;
 		}
 
-		return false;
+		MessageHeap.HeapPop(OutMessage, FSequenceComparer());
+		NextSequence++;
+
+		return true;
 	}
 
 	/**
@@ -69,7 +68,6 @@ public:
 	void Reset()
 	{
 		MessageHeap.Reset();
-
 		NextSequence = 1;
 	}
 
diff --git a/Engine/Plugins/Messaging/UdpMessaging/Source/UdpMessaging/Private/Transport/UdpMessageSegmenter.cpp b/Engine/Plugins/Messaging/UdpMessaging/Source/UdpMessaging/Private/Transport/UdpMessageSegmenter.cpp
--- a/Engine/Plugins/Messaging/UdpMessaging/Source/UdpMessaging/Private/Transport/UdpMessageSegmenter.cpp
+++ b/Engine/Plugins/Messaging/UdpMessaging/Source/UdpMessaging/Private/Transport/UdpMessageSegmenter.cpp
@@ -91,10 +91,8 @@ void FUdpMessageSegmenter::MarkAsSent(uint16 Segment)
 
 void FUdpMessageSegmenter::MarkForRetransmission(const TArray<uint16>& Segments)
 {
-	for (int32 Index = 0; Index < Segments.Num(); ++Index)
+	for (const auto& Segment : Segments)
 	{
-		uint16 Segment = Segments[Index];
-
 		if (Segment < PendingSegments.Num())
 		{
 			PendingSegments[Segment] = true;
diff --git a/Engine/Plugins/Messaging/UdpMessaging/Source/UdpMessaging/Private/Transport/UdpMessageTransport.h b/Engine/Plugins/Messaging/UdpMessaging/Source/UdpMessaging/Private/Transport/UdpMessageTransport.h
--- a/Engine/Plugins/Messaging/UdpMessaging/Source/UdpMessaging/Private/Transport/UdpMessageTransport.h
+++ b/Engine/Plugins/Messaging/UdpMessaging/Source/UdpMessaging/Private/Transport/UdpMessageTransport.h
@@ -18,9 +18,9 @@ public:
 	 * @param InMulticastEndpoint The multicast group endpoint to transport messages to.
 	 * @param InMulticastTtl The multicast time-to-live.
 	 */
-	FUdpMessageTransport( const FIPv4Endpoint& InLocalEndpoint, const FIPv4Endpoint& InMulticastEndpoint, uint8 InMulticastTtl );
+	FUdpMessageTransport(const FIPv4Endpoint& InLocalEndpoint, const FIPv4Endpoint& InMulticastEndpoint, uint8 InMulticastTtl);
 
-	/** Destructor. */
+	/** Virtual destructor. */
 	virtual ~FUdpMessageTransport();
 
 public:
@@ -49,21 +49,21 @@ public:
 
 	virtual bool StartTransport() override;
 	virtual void StopTransport() override;
-	virtual bool TransportMessage( const IMessageContextRef& Context, const TArray<FGuid>& Recipients ) override;
+	virtual bool TransportMessage(const IMessageContextRef& Context, const TArray<FGuid>& Recipients) override;
 
 private:
 
 	/** Handles received transport messages. */
-	void HandleProcessorMessageReassembled( const FUdpReassembledMessageRef& ReassembledMessage, const IMessageAttachmentPtr& Attachment, const FGuid& NodeId );
+	void HandleProcessorMessageReassembled(const FUdpReassembledMessageRef& ReassembledMessage, const IMessageAttachmentPtr& Attachment, const FGuid& NodeId);
 
 	/** Handles discovered transport endpoints. */
-	void HandleProcessorNodeDiscovered( const FGuid& DiscoveredNodeId )
+	void HandleProcessorNodeDiscovered(const FGuid& DiscoveredNodeId)
 	{
 		NodeDiscoveredDelegate.ExecuteIfBound(DiscoveredNodeId);
 	}
 
 	/** Handles lost transport endpoints. */
-	void HandleProcessorNodeLost( const FGuid& LostNodeId )
+	void HandleProcessorNodeLost(const FGuid& LostNodeId)
 	{
 		NodeLostDelegate.ExecuteIfBound(LostNodeId);
 	}
diff --git a/Engine/Plugins/Messaging/UdpMessaging/Source/UdpMessaging/Private/Transport/UdpSerializeMessageTask.cpp b/Engine/Plugins/Messaging/UdpMessaging/Source/UdpMessaging/Private/Transport/UdpSerializeMessageTask.cpp
--- a/Engine/Plugins/Messaging/UdpMessaging/Source/UdpMessaging/Private/Transport/UdpSerializeMessageTask.cpp
+++ b/Engine/Plugins/Messaging/UdpMessaging/Source/UdpMessaging/Private/Transport/UdpSerializeMessageTask.cpp
@@ -7,7 +7,7 @@
 /* FUdpSerializeMessageTask interface
  *****************************************************************************/
 
-void FUdpSerializeMessageTask::DoTask( ENamedThreads::Type CurrentThread, const FGraphEventRef& MyCompletionGraphEvent )
+void FUdpSerializeMessageTask::DoTask(ENamedThreads::Type CurrentThread, const FGraphEventRef& MyCompletionGraphEvent)
 {
 	if (MessageContext->IsValid())
 	{
@@ -15,9 +15,8 @@ void FUdpSerializeMessageTask::DoTask( ENamedThreads::Type CurrentThread, const
 		// a consistent wire format, if their implementations change. This allows us to sanity
 		// check the values during deserialization. @see FUdpDeserializeMessage::Deserialize()
 
-		FArchive& Archive = SerializedMessage.Get();
-
 		// serialize context
+		FArchive& Archive = SerializedMessage.Get();
 		{
 			const FName& MessageType = MessageContext->GetMessageType();
 			Archive << const_cast<FName&>(MessageType);
@@ -40,18 +39,16 @@ void FUdpSerializeMessageTask::DoTask( ENamedThreads::Type CurrentThread, const
 			int32 NumAnnotations = MessageContext->GetAnnotations().Num();
 			Archive << NumAnnotations;
 
-			for (TMap<FName, FString>::TConstIterator It(MessageContext->GetAnnotations()); It; ++It)
+			for (const auto& AnnotationPair : MessageContext->GetAnnotations())
 			{
-				Archive << const_cast<FName&>(It->Key);
-				Archive << const_cast<FString&>(It->Value);
+				Archive << const_cast<FName&>(AnnotationPair.Key);
+				Archive << const_cast<FString&>(AnnotationPair.Value);
 			}
 		}
 
 		// serialize message body
-		{
-			FJsonStructSerializerBackend Backend(Archive);
-			FStructSerializer::Serialize(MessageContext->GetMessage(), *MessageContext->GetMessageTypeInfo(), Backend);
-		}
+		FJsonStructSerializerBackend Backend(Archive);
+		FStructSerializer::Serialize(MessageContext->GetMessage(), *MessageContext->GetMessageTypeInfo(), Backend);
 
 		SerializedMessage->UpdateState(EUdpSerializedMessageState::Complete);
 	}
diff --git a/Engine/Plugins/Messaging/UdpMessaging/Source/UdpMessaging/Private/Tunnel/UdpMessageTunnel.cpp b/Engine/Plugins/Messaging/UdpMessaging/Source/UdpMessaging/Private/Tunnel/UdpMessageTunnel.cpp
--- a/Engine/Plugins/Messaging/UdpMessaging/Source/UdpMessaging/Private/Tunnel/UdpMessageTunnel.cpp
+++ b/Engine/Plugins/Messaging/UdpMessaging/Source/UdpMessaging/Private/Tunnel/UdpMessageTunnel.cpp
@@ -6,7 +6,7 @@
 /* FUdpMessageTunnel structors
  *****************************************************************************/
 
-FUdpMessageTunnel::FUdpMessageTunnel( const FIPv4Endpoint& InUnicastEndpoint, const FIPv4Endpoint& InMulticastEndpoint )
+FUdpMessageTunnel::FUdpMessageTunnel(const FIPv4Endpoint& InUnicastEndpoint, const FIPv4Endpoint& InMulticastEndpoint)
 	: Listener(nullptr)
 	, MulticastEndpoint(InMulticastEndpoint)
 	, Stopping(false)
@@ -103,35 +103,35 @@ void FUdpMessageTunnel::Exit()
 /* IUdpMessageTunnel interface
  *****************************************************************************/
 
-bool FUdpMessageTunnel::Connect( const FIPv4Endpoint& RemoteEndpoint )
+bool FUdpMessageTunnel::Connect(const FIPv4Endpoint& RemoteEndpoint)
 {
 	FSocket* Socket = FTcpSocketBuilder(TEXT("FUdpMessageTunnel.RemoteConnection"));
 
-	if (Socket != nullptr)
+	if (Socket == nullptr)
 	{
-		if (Socket->Connect(*RemoteEndpoint.ToInternetAddr()))
-		{
-			PendingConnections.Enqueue(MakeShareable(new FUdpMessageTunnelConnection(Socket, RemoteEndpoint)));
+		return false;
+	}
 
-			return true;
-		}
-		else
-		{
-			ISocketSubsystem::Get(PLATFORM_SOCKETSUBSYSTEM)->DestroySocket(Socket);
-		}
+	if (!Socket->Connect(*RemoteEndpoint.ToInternetAddr()))
+	{
+		ISocketSubsystem::Get(PLATFORM_SOCKETSUBSYSTEM)->DestroySocket(Socket);
+
+		return false;
 	}
 
-	return false;
+	PendingConnections.Enqueue(MakeShareable(new FUdpMessageTunnelConnection(Socket, RemoteEndpoint)));
+
+	return true;
 }
 
 
-int32 FUdpMessageTunnel::GetConnections( TArray<IUdpMessageTunnelConnectionPtr>& OutConnections )
+int32 FUdpMessageTunnel::GetConnections(TArray<IUdpMessageTunnelConnectionPtr>& OutConnections)
 {
 	FScopeLock Lock(&CriticalSection);
 
-	for (TArray<FUdpMessageTunnelConnectionPtr>::TConstIterator It(Connections); It; ++It)
+	for (const auto& Connection : Connections)
 	{
-		OutConnections.Add(*It);
+		OutConnections.Add(Connection);
 	}
 
 	return OutConnections.Num();
@@ -162,7 +162,7 @@ FSimpleDelegate& FUdpMessageTunnel::OnConnectionsChanged()
 }
 
 
-void FUdpMessageTunnel::StartServer( const FIPv4Endpoint& LocalEndpoint )
+void FUdpMessageTunnel::StartServer(const FIPv4Endpoint& LocalEndpoint)
 {
 	StopServer();
 
@@ -181,7 +181,7 @@ void FUdpMessageTunnel::StopServer()
 /* FUdpMessageTunnel implementation
  *****************************************************************************/
 
-void FUdpMessageTunnel::RemoveExpiredNodes( TMap<FGuid, FNodeInfo>& Nodes )
+void FUdpMessageTunnel::RemoveExpiredNodes(TMap<FGuid, FNodeInfo>& Nodes)
 {
 	for (TMap<FGuid, FNodeInfo>::TIterator It(Nodes); It; ++It)
 	{
@@ -250,7 +250,7 @@ void FUdpMessageTunnel::TcpToUdp()
 }
 
 
-void FUdpMessageTunnel::UdpToTcp( FSocket* Socket )
+void FUdpMessageTunnel::UdpToTcp(FSocket* Socket)
 {
 	TSharedRef<FInternetAddr> Sender = ISocketSubsystem::Get(PLATFORM_SOCKETSUBSYSTEM)->CreateInternetAddr();
 
@@ -321,7 +321,6 @@ void FUdpMessageTunnel::UpdateConnections()
 		if (!Connections[ConnectionIndex]->IsOpen())
 		{
 			Connections.RemoveAtSwap(ConnectionIndex);
-
 			ConnectionsChanged = true;
 		}
 	}
@@ -350,7 +349,7 @@ void FUdpMessageTunnel::UpdateConnections()
 /* FUdpMessageTunnel callbacks
  *****************************************************************************/
 
-bool FUdpMessageTunnel::HandleListenerConnectionAccepted( FSocket* ClientSocket, const FIPv4Endpoint& ClientEndpoint )
+bool FUdpMessageTunnel::HandleListenerConnectionAccepted(FSocket* ClientSocket, const FIPv4Endpoint& ClientEndpoint)
 {
 	PendingConnections.Enqueue(MakeShareable(new FUdpMessageTunnelConnection(ClientSocket, ClientEndpoint)));
 
diff --git a/Engine/Plugins/Messaging/UdpMessaging/Source/UdpMessaging/Private/Tunnel/UdpMessageTunnel.h b/Engine/Plugins/Messaging/UdpMessaging/Source/UdpMessaging/Private/Tunnel/UdpMessageTunnel.h
--- a/Engine/Plugins/Messaging/UdpMessaging/Source/UdpMessaging/Private/Tunnel/UdpMessageTunnel.h
+++ b/Engine/Plugins/Messaging/UdpMessaging/Source/UdpMessaging/Private/Tunnel/UdpMessageTunnel.h
@@ -17,16 +17,16 @@ class FUdpMessageTunnel
 	: FRunnable
 	, public IUdpMessageTunnel
 {
-	// Structure for transport node information
+	/** Structure for transport node information. */
 	struct FNodeInfo
 	{
-		// Holds the connection that owns the node (only used for remote nodes).
+		/** Holds the connection that owns the node (only used for remote nodes). */
 		FUdpMessageTunnelConnectionPtr Connection;
 
-		// Holds the node's IP endpoint (only used for local nodes).
+		/** Holds the node's IP endpoint (only used for local nodes). */
 		FIPv4Endpoint Endpoint;
 
-		// Holds the time at which the last datagram was received.
+		/** Holds the time at which the last datagram was received. */
 		FDateTime LastDatagramReceivedTime;
 	};
 
@@ -38,7 +38,7 @@ public:
 	 * @param InLocalEndpoint The local IP endpoint to receive messages on.
 	 * @param InMulticastEndpoint The multicast group endpoint to transport messages to.
 	 */
-	FUdpMessageTunnel( const FIPv4Endpoint& InUnicastEndpoint, const FIPv4Endpoint& InMulticastEndpoint );
+	FUdpMessageTunnel(const FIPv4Endpoint& InUnicastEndpoint, const FIPv4Endpoint& InMulticastEndpoint);
 
 	/** Virtual destructor. */
 	virtual ~FUdpMessageTunnel();
@@ -56,13 +56,13 @@ public:
 
 	// IUdpMessageTunnel interface
 
-	virtual bool Connect( const FIPv4Endpoint& RemoteEndpoint ) override;
-	virtual int32 GetConnections( TArray<IUdpMessageTunnelConnectionPtr>& OutConnections ) override;
+	virtual bool Connect(const FIPv4Endpoint& RemoteEndpoint) override;
+	virtual int32 GetConnections(TArray<IUdpMessageTunnelConnectionPtr>& OutConnections) override;
 	virtual uint64 GetTotalInboundBytes() const override;
 	virtual uint64 GetTotalOutboundBytes() const override;
 	virtual bool IsServerRunning() const override;
 	virtual FSimpleDelegate& OnConnectionsChanged() override;
-	virtual void StartServer( const FIPv4Endpoint& LocalEndpoint ) override;
+	virtual void StartServer(const FIPv4Endpoint& LocalEndpoint) override;
 	virtual void StopServer() override;
 
 protected:
@@ -72,17 +72,22 @@ protected:
 	 *
 	 * @param Nodes The collection of nodes to clean up.
 	 */
-	void RemoveExpiredNodes( TMap<FGuid, FNodeInfo>& Nodes );
+	void RemoveExpiredNodes(TMap<FGuid, FNodeInfo>& Nodes);
 
-	/** Receives all pending payloads from the tunnels and forwards them to the local message bus. */
+	/**
+	 * Receives all pending payloads from the tunnels and forwards them to the local message bus.
+	 *
+	 * @see UdpToTcp
+	 */
 	void TcpToUdp();
 
 	/**
 	 * Receives all buffered datagrams from the specified socket and forwards them to the tunnels.
 	 *
 	 * @param Socket The socket to receive from.
+	 * @see TcpToUdp
 	 */
-	void UdpToTcp( FSocket* Socket );
+	void UdpToTcp(FSocket* Socket);
 
 	/** Updates all active and pending connections. */
 	void UpdateConnections();
@@ -90,7 +95,7 @@ protected:
 private:
 
 	/** Callback for accepted connections to the local tunnel server. */
-	bool HandleListenerConnectionAccepted( FSocket* ClientSocket, const FIPv4Endpoint& ClientEndpoint );
+	bool HandleListenerConnectionAccepted(FSocket* ClientSocket, const FIPv4Endpoint& ClientEndpoint);
 
 private:
 
diff --git a/Engine/Plugins/Messaging/UdpMessaging/Source/UdpMessaging/Private/Tunnel/UdpMessageTunnelConnection.cpp b/Engine/Plugins/Messaging/UdpMessaging/Source/UdpMessaging/Private/Tunnel/UdpMessageTunnelConnection.cpp
--- a/Engine/Plugins/Messaging/UdpMessaging/Source/UdpMessaging/Private/Tunnel/UdpMessageTunnelConnection.cpp
+++ b/Engine/Plugins/Messaging/UdpMessaging/Source/UdpMessaging/Private/Tunnel/UdpMessageTunnelConnection.cpp
@@ -6,7 +6,7 @@
 /* FUdpMessageTunnelConnection structors
  *****************************************************************************/
 
-FUdpMessageTunnelConnection::FUdpMessageTunnelConnection( FSocket* InSocket, const FIPv4Endpoint& InRemoteEndpoint )
+FUdpMessageTunnelConnection::FUdpMessageTunnelConnection(FSocket* InSocket, const FIPv4Endpoint& InRemoteEndpoint)
 	: OpenedTime(FDateTime::UtcNow())
 	, RemoteEndpoint(InRemoteEndpoint)
 	, Socket(InSocket)
@@ -37,13 +37,13 @@ FUdpMessageTunnelConnection::~FUdpMessageTunnelConnection()
 /* FUdpMessageTunnelConnection interface
  *****************************************************************************/
 
-bool FUdpMessageTunnelConnection::Receive( FArrayReaderPtr& OutPayload )
+bool FUdpMessageTunnelConnection::Receive(FArrayReaderPtr& OutPayload)
 {
 	return Inbox.Dequeue(OutPayload);
 }
 
 
-bool FUdpMessageTunnelConnection::Send( const FArrayReaderPtr& Payload )
+bool FUdpMessageTunnelConnection::Send(const FArrayReaderPtr& Payload)
 {
 	if (IsOpen())
 	{
@@ -204,35 +204,39 @@ bool FUdpMessageTunnelConnection::ReceivePayloads()
 
 bool FUdpMessageTunnelConnection::SendPayloads()
 {
-	if (!Outbox.IsEmpty())
+	if (Outbox.IsEmpty())
 	{
-		if (Socket->Wait(ESocketWaitConditions::WaitForWrite, FTimespan::Zero()))
-		{
-			FArrayReaderPtr Payload;
-			Outbox.Dequeue(Payload);
-			int32 BytesSent = 0;
+		return true;
+	}
 
-			// send the payload size
-			FArrayWriter PayloadSizeData = FArrayWriter(true);
-			uint16 PayloadSize = Payload->Num();
-			PayloadSizeData << PayloadSize;
+	if (!Socket->Wait(ESocketWaitConditions::WaitForWrite, FTimespan::Zero()))
+	{
+		return true;
+	}
 
-			if (!Socket->Send(PayloadSizeData.GetData(), sizeof(uint16), BytesSent))
-			{
-				return false;
-			}
+	FArrayReaderPtr Payload;
+	Outbox.Dequeue(Payload);
+	int32 BytesSent = 0;
 
-			TotalBytesSent += BytesSent;
+	// send the payload size
+	FArrayWriter PayloadSizeData = FArrayWriter(true);
+	uint16 PayloadSize = Payload->Num();
+	PayloadSizeData << PayloadSize;
 
-			// send the payload
-			if (!Socket->Send(Payload->GetData(), Payload->Num(), BytesSent))
-			{
-				return false;
-			}
+	if (!Socket->Send(PayloadSizeData.GetData(), sizeof(uint16), BytesSent))
+	{
+		return false;
+	}
 
-			TotalBytesSent += BytesSent;
-		}
+	TotalBytesSent += BytesSent;
+
+	// send the payload
+	if (!Socket->Send(Payload->GetData(), Payload->Num(), BytesSent))
+	{
+		return false;
 	}
 
+	TotalBytesSent += BytesSent;
+
 	return true;
 }
