Back out revision 7 from //depot/UE4/Engine/Plugins/Messaging/UdpMessaging/Source/UdpMessaging/Private/Transport/Tests/UdpMessageSegmenterTest.cpp

diff --git a/Engine/Plugins/Messaging/UdpMessaging/Source/UdpMessaging/Private/Transport/Tests/UdpMessageSegmenterTest.cpp b/Engine/Plugins/Messaging/UdpMessaging/Source/UdpMessaging/Private/Transport/Tests/UdpMessageSegmenterTest.cpp
new file mode 100644
--- /dev/null
+++ b/Engine/Plugins/Messaging/UdpMessaging/Source/UdpMessaging/Private/Transport/Tests/UdpMessageSegmenterTest.cpp
@@ -0,0 +1,77 @@
+// Copyright 1998-2015 Epic Games, Inc. All Rights Reserved.
+
+#include "UdpMessagingPrivatePCH.h"
+#include "AutomationTest.h"
+
+
+IMPLEMENT_SIMPLE_AUTOMATION_TEST(FUdpMessageSegmenterTest, "System.Core.Messaging.Transports.Udp.UdpMessageSegmenter", EAutomationTestFlags::EditorContext | EAutomationTestFlags::EngineFilter)
+
+
+bool FUdpMessageSegmenterTest::RunTest( const FString& Parameters )
+{/*
+	const uint16 NumSegments = 512;
+	const uint16 TestSegmentSize = 1024;
+
+	// create a large message to segment
+	TSharedRef<FMessageData, ESPMode::ThreadSafe> TestData = MakeShareable(new FMessageData());
+
+	for (uint16 SegmentIndex = 0; SegmentIndex < NumSegments; ++SegmentIndex)
+	{
+		for (uint16 Offset = 0; Offset < TestSegmentSize / sizeof(uint16); ++Offset)
+		{
+			*TestData << SegmentIndex;
+		}
+	}
+
+	TestData->UpdateState(EMessageDataState::Complete);
+
+	// create and initialize segmenter
+	FUdpMessageSegmenter Segmenter = FUdpMessageSegmenter(TestData, TestSegmentSize);
+	Segmenter.Initialize();
+
+	// invariants
+	{
+		TestEqual(TEXT("The message size must match the actual message size"), Segmenter.GetMessageSize(), TestData->TotalSize());
+		TestEqual(TEXT("The total number of segments must match the actual number of segments in the message"), Segmenter.GetSegmentCount(), NumSegments);
+	}
+
+	// pre-conditions
+	{
+		TestEqual(TEXT("The initial number of pending segments must match the total number of segments in the message"), Segmenter.GetPendingSegmentsCount(), NumSegments);
+		TestFalse(TEXT("Segmentation of a non-empty message must not be complete initially"), Segmenter.IsComplete());
+	}
+
+	// peeking at next pending segment
+	{
+		TArray<uint8> OutData;
+		uint16 OutSegmentNumber;
+
+		Segmenter.GetNextPendingSegment(OutData, OutSegmentNumber);
+
+		TestEqual(TEXT("The number of pending segments must not change when peeking at a pending segment"), Segmenter.GetPendingSegmentsCount(), NumSegments);
+	}
+
+	uint16 GeneratedSegmentCount = 0;
+
+	// segmentation
+	{
+		TArray<uint8> OutData;
+		uint16 OutSegmentNumber;
+
+		while (Segmenter.GetNextPendingSegment(OutData, OutSegmentNumber))
+		{
+			Segmenter.MarkAsSent(OutSegmentNumber);
+
+			++GeneratedSegmentCount;
+		}
+	}
+
+	// post-conditions
+	{
+		TestEqual(TEXT("The number of generated segments must match the total number of segments in the message"), GeneratedSegmentCount, NumSegments);
+		TestEqual(TEXT("The number of pending segments must be zero after segmentation is complete"), Segmenter.GetPendingSegmentsCount(), (uint16)0);
+		TestTrue(TEXT("Segmentation must be complete when there are no more pending segments"), Segmenter.IsComplete());
+	}
+*/
+	return true;
+}
