UdpMessaging: Made serialization task fire-and-forget

diff --git a/Engine/Plugins/Messaging/UdpMessaging/Source/UdpMessaging/Private/Transport/UdpSerializeMessageTask.cpp b/Engine/Plugins/Messaging/UdpMessaging/Source/UdpMessaging/Private/Transport/UdpSerializeMessageTask.cpp
--- a/Engine/Plugins/Messaging/UdpMessaging/Source/UdpMessaging/Private/Transport/UdpSerializeMessageTask.cpp
+++ b/Engine/Plugins/Messaging/UdpMessaging/Source/UdpMessaging/Private/Transport/UdpSerializeMessageTask.cpp
@@ -73,5 +73,5 @@ TStatId FUdpSerializeMessageTask::GetStatId() const
 
 ESubsequentsMode::Type FUdpSerializeMessageTask::GetSubsequentsMode() 
 { 
-	return ESubsequentsMode::TrackSubsequents; 
+	return ESubsequentsMode::FireAndForget; 
 }
