UdpMessaging: Added getter for data array of serialized messages

diff --git a/Engine/Plugins/Messaging/UdpMessaging/Source/UdpMessaging/Private/Transport/UdpSerializedMessage.h b/Engine/Plugins/Messaging/UdpMessaging/Source/UdpMessaging/Private/Transport/UdpSerializedMessage.h
--- a/Engine/Plugins/Messaging/UdpMessaging/Source/UdpMessaging/Private/Transport/UdpSerializedMessage.h
+++ b/Engine/Plugins/Messaging/UdpMessaging/Source/UdpMessaging/Private/Transport/UdpSerializedMessage.h
@@ -31,7 +31,7 @@ public:
 
 	/** Default constructor. */
 	FUdpSerializedMessage()
-		: FMemoryWriter(Data, true)
+		: FMemoryWriter(DataArray, true)
 		, State(EUdpSerializedMessageState::Incomplete)
 	{ }
 
@@ -46,14 +46,25 @@ public:
 	 */
 	FArchive* CreateReader()
 	{
-		return new FMemoryReader(Data, true);
+		return new FMemoryReader(DataArray, true);
+	}
+
+	/**
+	 * Get the serialized message data.
+	 *
+	 * @return Byte array of message data.
+	 * @see GetState
+	 */
+	const TArray<uint8>& GetDataArray()
+	{
+		return DataArray;
 	}
 
 	/**
 	 * Gets the state of the message data.
 	 *
 	 * @return Message data state.
-	 * @see OnStateChanged
+	 * @see GetData, OnStateChanged, UpdateState
 	 */
 	EUdpSerializedMessageState GetState() const
 	{
@@ -64,7 +75,7 @@ public:
 	 * Returns a delegate that is executed when the message data's state changed.
 	 *
 	 * @return The delegate.
-	 * @see GetState
+	 * @see GetState, UpdateState
 	 */
 	FSimpleDelegate& OnStateChanged()
 	{
@@ -75,6 +86,7 @@ public:
 	 * Updates the state of this message data.
 	 *
 	 * @param InState The state to set.
+	 * @see GetState, OnStateChanged
 	 */
 	void UpdateState(EUdpSerializedMessageState InState)
 	{
@@ -84,8 +96,8 @@ public:
 
 private:
 
-	/** Holds the data. */
-	TArray<uint8> Data;
+	/** Holds the serialized data. */
+	TArray<uint8> DataArray;
 
 	/** Holds the message data state. */
 	EUdpSerializedMessageState State;
