UdpMessaging: Added unit test for message serialization

diff --git a/Engine/Plugins/Messaging/UdpMessaging/Source/UdpMessaging/Private/Tests/UdpMessagingTestTypes.h b/Engine/Plugins/Messaging/UdpMessaging/Source/UdpMessaging/Private/Tests/UdpMessagingTestTypes.h
new file mode 100644
--- /dev/null
+++ b/Engine/Plugins/Messaging/UdpMessaging/Source/UdpMessaging/Private/Tests/UdpMessagingTestTypes.h
@@ -0,0 +1,68 @@
+// Copyright 1998-2015 Epic Games, Inc. All Rights Reserved.
+
+#pragma once
+
+#include "UdpMessagingTestTypes.generated.h"
+
+
+USTRUCT()
+struct FUdpMockMessage
+{
+	GENERATED_USTRUCT_BODY()
+
+	UPROPERTY()
+	int32 Foo;
+
+	FUdpMockMessage()
+		: Foo(123)
+	{ }
+};
+
+
+class FUdpMockMessageContext
+	: public IMessageContext
+{
+public:
+
+	FUdpMockMessageContext(FUdpMockMessage* InMessage)
+		: Expiration(FDateTime::MaxValue())
+		, Message(InMessage)
+		, Scope(EMessageScope::Network)
+		, SenderThread(ENamedThreads::AnyThread)
+		, TimeSent(FDateTime(2015, 9, 17, 10, 59, 23, 666))
+		, TypeInfo(FUdpMockMessage::StaticStruct())
+	{
+		FMessageAddress::Parse(TEXT("11111111-22222222-33333333-44444444"), Sender);
+	}
+
+public:
+
+	// IMessageContext interface
+
+	virtual const TMap<FName, FString>& GetAnnotations() const override { return Annotations; }
+	virtual TSharedPtr<IMessageAttachment, ESPMode::ThreadSafe> GetAttachment() const override { return Attachment; }
+	virtual const FDateTime& GetExpiration() const override { return Expiration; }
+	virtual const void* GetMessage() const override { return Message; }
+	virtual const TWeakObjectPtr<UScriptStruct>& GetMessageTypeInfo() const override { return TypeInfo; }
+	virtual TSharedPtr<IMessageContext, ESPMode::ThreadSafe> GetOriginalContext() const override { return OriginalContext; }
+	virtual const TArray<FMessageAddress>& GetRecipients() const override { return Recipients; }
+	virtual EMessageScope GetScope() const override { return Scope; }
+	virtual const FMessageAddress& GetSender() const override { return Sender; }
+	virtual ENamedThreads::Type GetSenderThread() const override { return SenderThread; }
+	virtual const FDateTime& GetTimeForwarded() const override { return TimeSent; }
+	virtual const FDateTime& GetTimeSent() const override { return TimeSent; }
+
+private:
+
+	TMap<FName, FString> Annotations;
+	IMessageAttachmentPtr Attachment;
+	FDateTime Expiration;
+	void* Message;
+	IMessageContextPtr OriginalContext;
+	TArray<FMessageAddress> Recipients;
+	EMessageScope Scope;
+	FMessageAddress Sender;
+	ENamedThreads::Type SenderThread;
+	FDateTime TimeSent;
+	TWeakObjectPtr<UScriptStruct> TypeInfo;
+};
diff --git a/Engine/Plugins/Messaging/UdpMessaging/Source/UdpMessaging/Private/Tests/UdpSerializeMessageTaskTest.cpp b/Engine/Plugins/Messaging/UdpMessaging/Source/UdpMessaging/Private/Tests/UdpSerializeMessageTaskTest.cpp
new file mode 100644
--- /dev/null
+++ b/Engine/Plugins/Messaging/UdpMessaging/Source/UdpMessaging/Private/Tests/UdpSerializeMessageTaskTest.cpp
@@ -0,0 +1,92 @@
+// Copyright 1998-2015 Epic Games, Inc. All Rights Reserved.
+
+#include "UdpMessagingPrivatePCH.h"
+#include "AutomationTest.h"
+#include "UdpMessagingTestTypes.h"
+
+
+
+IMPLEMENT_SIMPLE_AUTOMATION_TEST(FUdpSerializeMessageTaskTest, "System.Core.Messaging.Transports.Udp.UdpSerializedMessage", EAutomationTestFlags::ApplicationContextMask | EAutomationTestFlags::EngineFilter)
+
+
+namespace UdpSerializeMessageTaskTest
+{
+	const int64 NumMessages = 100000;
+	const FTimespan MaxWaitTime(0, 0, 5);
+
+	int64 CompletedMessages = 0;
+	int64 FailedMessages = 0;
+	FUdpSerializedMessagePtr ReferenceMessage;
+}
+
+
+void HandleSerializedMessageStateChanged(FUdpSerializedMessageRef SerializedMessage)
+{
+	using namespace UdpSerializeMessageTaskTest;
+
+	FPlatformAtomics::InterlockedIncrement(&CompletedMessages);
+
+	const TArray<uint8>& ReferenceDataArray = ReferenceMessage->GetDataArray();
+	const TArray<uint8>& SerializedDataArray = SerializedMessage->GetDataArray();
+
+	if (ReferenceDataArray.Num() != SerializedDataArray.Num())
+	{
+		FPlatformAtomics::InterlockedIncrement(&FailedMessages);
+	}
+	else
+	{
+		const void* ReferenceData = ReferenceDataArray.GetData();
+		const void* SerializedData = SerializedDataArray.GetData();
+
+		if (FMemory::Memcmp(ReferenceData, SerializedData, ReferenceDataArray.Num()) != 0)
+		{
+			FPlatformAtomics::InterlockedIncrement(&FailedMessages);
+		}
+	}
+}
+
+
+bool FUdpSerializeMessageTaskTest::RunTest(const FString& Parameters)
+{
+	using namespace UdpSerializeMessageTaskTest;
+
+	CompletedMessages = 0;
+	FailedMessages = 0;
+	ReferenceMessage = MakeShareable(new FUdpSerializedMessage());
+
+	FUdpMockMessage* Message = new FUdpMockMessage();
+	IMessageContextRef Context = MakeShareable(new FUdpMockMessageContext(Message));
+
+	// synchronous reference serialization
+	FUdpSerializeMessageTask ReferenceTask(Context, ReferenceMessage.ToSharedRef());
+	{
+		ReferenceTask.DoTask(FTaskGraphInterface::Get().GetCurrentThreadIfKnown(), FGraphEventRef());
+	}
+
+	// stress test
+	for (int32 TestIndex = 0; TestIndex < NumMessages; ++TestIndex)
+	{
+		FUdpSerializedMessageRef SerializedMessage = MakeShareable(new FUdpSerializedMessage());
+		{
+			SerializedMessage->OnStateChanged().BindStatic(&HandleSerializedMessageStateChanged, SerializedMessage);
+		}
+
+		TGraphTask<FUdpSerializeMessageTask>::CreateTask().ConstructAndDispatchWhenReady(Context, SerializedMessage);
+	}
+
+	FDateTime StartTime = FDateTime::UtcNow();
+
+	while ((CompletedMessages < NumMessages) && ((FDateTime::UtcNow() - StartTime) < MaxWaitTime))
+	{
+		FPlatformProcess::Sleep(0.0f);
+	}
+
+	TestEqual(TEXT("The number of completed messages must equal the total number of messages"), CompletedMessages, NumMessages);
+	TestEqual(TEXT("There must be no failed messages"), FailedMessages, (int64)0);
+
+	// clean up
+	FMemory::Free(Message);
+	ReferenceMessage.Reset();
+
+	return ((CompletedMessages == NumMessages) && (FailedMessages == 0));
+}
