UdpMessaging: Revived and implemented unit test for message segmentation

diff --git a/Engine/Plugins/Messaging/UdpMessaging/Source/UdpMessaging/Private/Tests/UdpMessageSegmenterTest.cpp b/Engine/Plugins/Messaging/UdpMessaging/Source/UdpMessaging/Private/Tests/UdpMessageSegmenterTest.cpp
new file mode 100644
--- /dev/null
+++ b/Engine/Plugins/Messaging/UdpMessaging/Source/UdpMessaging/Private/Tests/UdpMessageSegmenterTest.cpp
@@ -0,0 +1,124 @@
+// Copyright 1998-2015 Epic Games, Inc. All Rights Reserved.
+
+#include "UdpMessagingPrivatePCH.h"
+#include "AutomationTest.h"
+
+
+IMPLEMENT_SIMPLE_AUTOMATION_TEST(FUdpMessageSegmenterTest, "System.Core.Messaging.Transports.Udp.UdpMessageSegmenter", EAutomationTestFlags::ApplicationContextMask | EAutomationTestFlags::EngineFilter)
+
+
+void RunSegmentationTest(FAutomationTestBase& Test, uint32 MessageSize, uint16 SegmentSize)
+{
+	check(MessageSize < 255u * SegmentSize);
+
+	uint16 NumSegments = (MessageSize + SegmentSize - 1) / SegmentSize;
+	Test.AddLogItem(FString::Printf(TEXT("Segmenting message of size %i with %i segments of size %i..."), MessageSize, NumSegments, SegmentSize));
+
+	// create a large message to segment
+	FUdpSerializedMessageRef Message = MakeShareable(new FUdpSerializedMessage());
+
+	for (uint8 SegmentIndex = 0; SegmentIndex < NumSegments; ++SegmentIndex)
+	{
+		// write segment index into each byte of the current segment
+		for (uint16 Offset = 0; (Offset < SegmentSize) && (Message->TotalSize() < MessageSize); ++Offset)
+		{
+			*Message << SegmentIndex;
+		}
+	}
+
+	Message->UpdateState(EUdpSerializedMessageState::Complete);
+
+	// create and initialize segmenter
+	FUdpMessageSegmenter Segmenter = FUdpMessageSegmenter(Message, SegmentSize);
+	Segmenter.Initialize();
+
+	// invariants
+	{
+		Test.TestEqual(TEXT("The message size must match the actual message size"), Segmenter.GetMessageSize(), Message->TotalSize());
+		Test.TestEqual(TEXT("The total number of segments must match the actual number of segments in the message"), Segmenter.GetSegmentCount(), NumSegments);
+	}
+
+	// pre-conditions
+	{
+		Test.TestEqual(TEXT("The initial number of pending segments must match the total number of segments in the message"), Segmenter.GetPendingSegmentsCount(), NumSegments);
+		Test.TestFalse(TEXT("Segmentation of a non-empty message must not be complete initially"), Segmenter.IsComplete());
+	}
+
+	// peeking at next pending segment
+	{
+		TArray<uint8> OutData;
+		uint16 OutSegmentNumber;
+
+		Segmenter.GetNextPendingSegment(OutData, OutSegmentNumber);
+
+		Test.TestEqual(TEXT("The number of pending segments must not change when peeking at a pending segment"), Segmenter.GetPendingSegmentsCount(), NumSegments);
+	}
+
+	uint16 GeneratedSegmentCount = 0;
+
+	// do the segmentation
+	{
+		TArray<uint8> OutData;
+		uint16 OutSegmentNumber = 0;
+		int32 NumInvalidSegments = 0;
+
+		while (Segmenter.GetNextPendingSegment(OutData, OutSegmentNumber))
+		{
+			Segmenter.MarkAsSent(OutSegmentNumber);
+			++GeneratedSegmentCount;
+
+			// verify segment size
+			int32 ExpectedSegmentSize = (OutSegmentNumber == (NumSegments - 1))
+				? MessageSize % SegmentSize
+				: SegmentSize;
+
+			if (ExpectedSegmentSize == 0)
+			{
+				ExpectedSegmentSize = SegmentSize;
+			}
+
+			if (OutData.Num() != ExpectedSegmentSize)
+			{
+				++NumInvalidSegments;
+
+				continue;
+			}
+
+			// verify segment data
+			for (const auto& Element : OutData)
+			{
+				if (OutSegmentNumber != Element)
+				{
+					++NumInvalidSegments;
+				}
+			}
+		}
+
+		Test.TestEqual(TEXT("The number of generated segments must match the total number of segments in the message"), GeneratedSegmentCount, NumSegments);
+		Test.TestEqual(TEXT("The number of invalid segments must be zero"), NumInvalidSegments, (int32)0);
+	}
+
+	// post-conditions
+	{
+		Test.TestEqual(TEXT("The number of pending segments must be zero after segmentation is complete"), Segmenter.GetPendingSegmentsCount(), (uint16)0);
+		Test.TestTrue(TEXT("Segmentation must be complete when there are no more pending segments"), Segmenter.IsComplete());
+	}
+}
+
+
+bool FUdpMessageSegmenterTest::RunTest(const FString& Parameters)
+{
+	// single partial segment
+	RunSegmentationTest(*this, 123, 1024);
+
+	// single full segment
+	RunSegmentationTest(*this, 1024, 1024);
+
+	// multiple segments with partial last segment
+	RunSegmentationTest(*this, 100000, 1024);
+
+	// multiple segments without partial last segment
+	RunSegmentationTest(*this, 131072, 1024);
+
+	return true;
+}
diff --git a/Engine/Plugins/Messaging/UdpMessaging/Source/UdpMessaging/Private/Transport/Tests/UdpMessageSegmenterTest.cpp b/Engine/Plugins/Messaging/UdpMessaging/Source/UdpMessaging/Private/Transport/Tests/UdpMessageSegmenterTest.cpp
deleted file mode 100644
