UdpMessaging: Started to implement stress test for message transport

diff --git a/Engine/Plugins/Messaging/UdpMessaging/Source/UdpMessaging/Private/Tests/UdpMessagingTestTypes.h b/Engine/Plugins/Messaging/UdpMessaging/Source/UdpMessaging/Private/Tests/UdpMessagingTestTypes.h
--- a/Engine/Plugins/Messaging/UdpMessaging/Source/UdpMessaging/Private/Tests/UdpMessagingTestTypes.h
+++ b/Engine/Plugins/Messaging/UdpMessaging/Source/UdpMessaging/Private/Tests/UdpMessagingTestTypes.h
@@ -11,11 +11,17 @@ struct FUdpMockMessage
 	GENERATED_USTRUCT_BODY()
 
 	UPROPERTY()
-	int32 Foo;
+	TArray<uint8> Data;
 
 	FUdpMockMessage()
-		: Foo(123)
-	{ }
+	{
+		Data.AddUninitialized(64);
+	}
+
+	FUdpMockMessage(int32 DataSize)
+	{
+		Data.AddUninitialized(DataSize);
+	}
 };
 
 
@@ -35,6 +41,19 @@ public:
 		FMessageAddress::Parse(TEXT("11111111-22222222-33333333-44444444"), Sender);
 	}
 
+	~FUdpMockMessageContext()
+	{
+		if (Message != nullptr)
+		{
+			if (TypeInfo.IsValid())
+			{
+				TypeInfo->DestroyStruct(Message);
+			}
+
+			FMemory::Free(Message);
+		}
+	}
+
 public:
 
 	// IMessageContext interface
diff --git a/Engine/Plugins/Messaging/UdpMessaging/Source/UdpMessaging/Private/Transport/UdpMessageProcessor.cpp b/Engine/Plugins/Messaging/UdpMessaging/Source/UdpMessaging/Private/Transport/UdpMessageProcessor.cpp
--- a/Engine/Plugins/Messaging/UdpMessaging/Source/UdpMessaging/Private/Transport/UdpMessageProcessor.cpp
+++ b/Engine/Plugins/Messaging/UdpMessaging/Source/UdpMessaging/Private/Transport/UdpMessageProcessor.cpp
@@ -183,8 +183,13 @@ void FUdpMessageProcessor::ConsumeInboundSegments()
 		{
 			FNodeInfo& NodeInfo = KnownNodes.FindOrAdd(Header.SenderNodeId);
 
+			if (!NodeInfo.NodeId.IsValid())
+			{
+				NodeInfo.NodeId = Header.SenderNodeId;
+				NodeDiscoveredDelegate.ExecuteIfBound(NodeInfo.NodeId);
+			}
+
 			NodeInfo.Endpoint = Segment.Sender;
-			NodeInfo.NodeId = Header.SenderNodeId;
 
 			switch (Header.SegmentType)
 			{
@@ -232,6 +237,11 @@ void FUdpMessageProcessor::ConsumeOutboundMessages()
 
 	while (OutboundMessages.Dequeue(OutboundMessage))
 	{
+		if (OutboundMessage.SerializedMessage->TotalSize() > 1024 * 65536)
+		{
+			continue;
+		}
+
 		++LastSentMessage;
 
 		FNodeInfo& RecipientNodeInfo = KnownNodes.FindOrAdd(OutboundMessage.RecipientId);
diff --git a/Engine/Plugins/Messaging/UdpMessaging/Source/UdpMessaging/Private/Transport/UdpMessageTransport.cpp b/Engine/Plugins/Messaging/UdpMessaging/Source/UdpMessaging/Private/Transport/UdpMessageTransport.cpp
--- a/Engine/Plugins/Messaging/UdpMessaging/Source/UdpMessaging/Private/Transport/UdpMessageTransport.cpp
+++ b/Engine/Plugins/Messaging/UdpMessaging/Source/UdpMessaging/Private/Transport/UdpMessageTransport.cpp
@@ -66,9 +66,11 @@ bool FUdpMessageTransport::StartTransport()
 	FTimespan ThreadWaitTime = FTimespan::FromMilliseconds(100);
 
 	MessageProcessor = new FUdpMessageProcessor(UnicastSocket, FGuid::NewGuid(), MulticastEndpoint);
-	MessageProcessor->OnMessageReassembled().BindRaw(this, &FUdpMessageTransport::HandleProcessorMessageReassembled);
-	MessageProcessor->OnNodeDiscovered().BindRaw(this, &FUdpMessageTransport::HandleProcessorNodeDiscovered);
-	MessageProcessor->OnNodeLost().BindRaw(this, &FUdpMessageTransport::HandleProcessorNodeLost);
+	{
+		MessageProcessor->OnMessageReassembled().BindRaw(this, &FUdpMessageTransport::HandleProcessorMessageReassembled);
+		MessageProcessor->OnNodeDiscovered().BindRaw(this, &FUdpMessageTransport::HandleProcessorNodeDiscovered);
+		MessageProcessor->OnNodeLost().BindRaw(this, &FUdpMessageTransport::HandleProcessorNodeLost);
+	}
 
 	if (MulticastSocket != nullptr)
 	{
@@ -77,7 +79,9 @@ bool FUdpMessageTransport::StartTransport()
 	}
 
 	UnicastReceiver = new FUdpSocketReceiver(UnicastSocket, ThreadWaitTime, TEXT("UdpMessageUnicastReceiver"));
-	UnicastReceiver->OnDataReceived().BindRaw(this, &FUdpMessageTransport::HandleSocketDataReceived);
+	{
+		UnicastReceiver->OnDataReceived().BindRaw(this, &FUdpMessageTransport::HandleSocketDataReceived);
+	}
 
 	return true;
 }
