UdpMessaging: Started to implement stress test for message transport

diff --git a/Engine/Plugins/Messaging/UdpMessaging/Source/UdpMessaging/Private/Tests/UdpMessageTransportTest.cpp b/Engine/Plugins/Messaging/UdpMessaging/Source/UdpMessaging/Private/Tests/UdpMessageTransportTest.cpp
new file mode 100644
--- /dev/null
+++ b/Engine/Plugins/Messaging/UdpMessaging/Source/UdpMessaging/Private/Tests/UdpMessageTransportTest.cpp
@@ -0,0 +1,156 @@
+// Copyright 1998-2015 Epic Games, Inc. All Rights Reserved.
+
+#include "UdpMessagingPrivatePCH.h"
+#include "AutomationTest.h"
+#include "UdpMessagingTestTypes.h"
+
+
+IMPLEMENT_SIMPLE_AUTOMATION_TEST(FUdpMessageTransportTest, "System.Core.Messaging.Transports.Udp.UdpMessageTransport (takes ~2 minutes!)", EAutomationTestFlags::ApplicationContextMask | EAutomationTestFlags::EngineFilter)
+
+
+class FUdpMessageTransportTestState
+{
+public:
+
+	FUdpMessageTransportTestState(FAutomationTestBase& InTest, const FIPv4Endpoint& UnicastEndpoint, const FIPv4Endpoint& MulticastEndpoint, uint8 MulticastTimeToLive)
+		: NumReceivedMessages(0)
+		, Test(InTest)
+	{
+		Transport = MakeShareable(new FUdpMessageTransport(UnicastEndpoint, MulticastEndpoint, MulticastTimeToLive));
+		{
+			Transport->OnMessageReceived().BindRaw(this, &FUdpMessageTransportTestState::HandleTransportMessageReceived);
+			Transport->OnNodeDiscovered().BindRaw(this, &FUdpMessageTransportTestState::HandleTransportNodeDiscovered);
+			Transport->OnNodeLost().BindRaw(this, &FUdpMessageTransportTestState::HandleTransportNodeLost);
+		}
+	}
+
+public:
+
+	const TArray<FGuid>& GetDiscoveredNodes() const
+	{
+		return DiscoveredNodes;
+	}
+
+	const TArray<FGuid>& GetLostNodes() const
+	{
+		return LostNodes;
+	}
+
+	int32 GetNumReceivedMessages() const
+	{
+		return NumReceivedMessages;
+	}
+
+	bool Publish(const IMessageContextRef& Context)
+	{
+		return Transport->TransportMessage(Context, TArray<FGuid>());
+	}
+
+	bool Start()
+	{
+		return Transport->StartTransport();
+	}
+
+	void Stop()
+	{
+		Transport->StopTransport();
+	}
+
+private:
+
+	void HandleTransportMessageReceived(const IMessageContextRef& MessageContext, const FGuid& NodeId)
+	{
+		FPlatformAtomics::InterlockedIncrement(&NumReceivedMessages);
+	}
+
+	void HandleTransportNodeDiscovered(const FGuid& NodeId)
+	{
+		DiscoveredNodes.Add(NodeId);
+	}
+
+	void HandleTransportNodeLost(const FGuid& NodeId)
+	{
+		LostNodes.Add(NodeId);
+	}
+
+private:
+
+	TArray<FGuid> DiscoveredNodes;
+	TArray<FGuid> LostNodes;
+	int64 NumReceivedMessages;
+	FAutomationTestBase& Test;
+	IMessageTransportPtr Transport;
+};
+
+
+bool FUdpMessageTransportTest::RunTest(const FString& Parameters)
+{
+	const FIPv4Endpoint MulticastEndpoint(FIPv4Address(231, 0, 0, 1), 7777);
+	const FIPv4Endpoint UnicastEndpoint = FIPv4Endpoint::Any;
+	const uint8 MulticastTimeToLive = 0;
+	const int64 NumTestMessages = 10000;
+	const int32 MessageSize = 1280;
+
+	// create transports
+	FUdpMessageTransportTestState Transport1(*this, UnicastEndpoint, MulticastEndpoint, MulticastTimeToLive);
+	const auto& DiscoveredNodes1 = Transport1.GetDiscoveredNodes();
+	
+	FUdpMessageTransportTestState Transport2(*this, UnicastEndpoint, MulticastEndpoint, MulticastTimeToLive);
+	const auto& DiscoveredNodes2 = Transport2.GetDiscoveredNodes();
+
+	// test transport node discovery
+	Transport1.Start();
+	FPlatformProcess::Sleep(3.0f);
+
+	TestTrue(TEXT("A single message transport must not discover any remote nodes"), DiscoveredNodes1.Num() == 0);
+
+	Transport2.Start();
+	FPlatformProcess::Sleep(3.0f);
+
+	if (DiscoveredNodes1.Num() == 0)
+	{
+		AddError(TEXT("The first transport did not discover any nodes"));
+
+		return false;
+	}
+
+	if (DiscoveredNodes2.Num() == 0)
+	{
+		AddError(TEXT("The second transport did not discover any nodes"));
+
+		return false;
+	}
+
+	TestTrue(TEXT("The first transport must discover exactly one node"), DiscoveredNodes1.Num() == 1);
+	TestTrue(TEXT("The second transport must discover exactly one node"), DiscoveredNodes2.Num() == 1);
+	TestTrue(TEXT("The discovered node IDs must be valid"), DiscoveredNodes1[0].IsValid() && DiscoveredNodes2[0].IsValid());
+	TestTrue(TEXT("The discovered node IDs must be unique"), DiscoveredNodes1[0] != DiscoveredNodes2[0]);
+
+	if (HasAnyErrors())
+	{
+		return false;
+	}
+
+	// stress test message sending
+	FDateTime StartTime = FDateTime::UtcNow();
+
+	for (int32 Count = 0; Count < NumTestMessages; ++Count)
+	{
+		FUdpMockMessage* Message = new FUdpMockMessage(MessageSize);
+		IMessageContextRef Context = MakeShareable(new FUdpMockMessageContext(Message));
+
+		Transport1.Publish(Context);
+	}
+
+	AddLogItem(FString::Printf(TEXT("Sent %i messages in %s"), NumTestMessages, *(FDateTime::UtcNow() - StartTime).ToString()));
+
+	while ((Transport2.GetNumReceivedMessages() < NumTestMessages) && ((FDateTime::UtcNow() - StartTime) < FTimespan::FromSeconds(120.0)))
+	{
+		FPlatformProcess::Sleep(0.0f);
+	}
+
+	AddLogItem(FString::Printf(TEXT("Received %i messages in %s"), Transport2.GetNumReceivedMessages(), *(FDateTime::UtcNow() - StartTime).ToString()));
+	TestTrue(TEXT("All sent messages must have been received"), Transport2.GetNumReceivedMessages() == NumTestMessages);
+
+	return true;
+}
