Fixed WinXP build

diff --git a/Engine/Plugins/Messaging/UdpMessaging/Source/UdpMessaging/Private/Tests/UdpMessageTransportTest.cpp b/Engine/Plugins/Messaging/UdpMessaging/Source/UdpMessaging/Private/Tests/UdpMessageTransportTest.cpp
--- a/Engine/Plugins/Messaging/UdpMessaging/Source/UdpMessaging/Private/Tests/UdpMessageTransportTest.cpp
+++ b/Engine/Plugins/Messaging/UdpMessaging/Source/UdpMessaging/Private/Tests/UdpMessageTransportTest.cpp
@@ -77,7 +77,7 @@ private:
 
 	TArray<FGuid> DiscoveredNodes;
 	TArray<FGuid> LostNodes;
-	int64 NumReceivedMessages;
+	int32 NumReceivedMessages;
 	FAutomationTestBase& Test;
 	IMessageTransportPtr Transport;
 };
@@ -88,7 +88,7 @@ bool FUdpMessageTransportTest::RunTest(const FString& Parameters)
 	const FIPv4Endpoint MulticastEndpoint(FIPv4Address(231, 0, 0, 1), 7777);
 	const FIPv4Endpoint UnicastEndpoint = FIPv4Endpoint::Any;
 	const uint8 MulticastTimeToLive = 0;
-	const int64 NumTestMessages = 10000;
+	const int32 NumTestMessages = 10000;
 	const int32 MessageSize = 1280;
 
 	// create transports
diff --git a/Engine/Plugins/Messaging/UdpMessaging/Source/UdpMessaging/Private/Tests/UdpSerializeMessageTaskTest.cpp b/Engine/Plugins/Messaging/UdpMessaging/Source/UdpMessaging/Private/Tests/UdpSerializeMessageTaskTest.cpp
--- a/Engine/Plugins/Messaging/UdpMessaging/Source/UdpMessaging/Private/Tests/UdpSerializeMessageTaskTest.cpp
+++ b/Engine/Plugins/Messaging/UdpMessaging/Source/UdpMessaging/Private/Tests/UdpSerializeMessageTaskTest.cpp
@@ -11,11 +11,11 @@ IMPLEMENT_SIMPLE_AUTOMATION_TEST(FUdpSerializeMessageTaskTest, "System.Core.Mess
 
 namespace UdpSerializeMessageTaskTest
 {
-	const int64 NumMessages = 100000;
+	const int32 NumMessages = 100000;
 	const FTimespan MaxWaitTime(0, 0, 5);
 
-	int64 CompletedMessages = 0;
-	int64 FailedMessages = 0;
+	int32 CompletedMessages = 0;
+	int32 FailedMessages = 0;
 	FUdpSerializedMessagePtr ReferenceMessage;
 }
 
@@ -82,7 +82,7 @@ bool FUdpSerializeMessageTaskTest::RunTest(const FString& Parameters)
 	}
 
 	TestEqual(TEXT("The number of completed messages must equal the total number of messages"), CompletedMessages, NumMessages);
-	TestEqual(TEXT("There must be no failed messages"), FailedMessages, (int64)0);
+	TestEqual(TEXT("There must be no failed messages"), FailedMessages, (int32)0);
 
 	// clean up
 	FMemory::Free(Message);
