UdpMessaging: Added console commands for basic diagnostics

Enter 'udpmessaging' into command console for details.

diff --git a/Engine/Plugins/Messaging/UdpMessaging/Source/UdpMessaging/Private/UdpMessagingModule.cpp b/Engine/Plugins/Messaging/UdpMessaging/Source/UdpMessaging/Private/UdpMessagingModule.cpp
--- a/Engine/Plugins/Messaging/UdpMessaging/Source/UdpMessaging/Private/UdpMessagingModule.cpp
+++ b/Engine/Plugins/Messaging/UdpMessaging/Source/UdpMessaging/Private/UdpMessagingModule.cpp
@@ -16,8 +16,99 @@ DEFINE_LOG_CATEGORY(LogUdpMessaging);
  * Implements the UdpMessagingModule module.
  */
 class FUdpMessagingModule
-	: public IModuleInterface
+	: public FSelfRegisteringExec
+	, public IModuleInterface
 {
+public:
+
+	// FSelfRegisteringExec interface
+
+	virtual bool Exec(UWorld* InWorld, const TCHAR* Cmd, FOutputDevice& Ar) override
+	{
+		if (!FParse::Command(&Cmd, TEXT("UDPMESSAGING")))
+		{
+			return false;
+		}
+
+		if (FParse::Command(&Cmd, TEXT("STATUS")))
+		{
+			UUdpMessagingSettings* Settings = GetMutableDefault<UUdpMessagingSettings>();
+
+			// bridge status
+			if (MessageBridge.IsValid())
+			{
+				if (MessageBridge->IsEnabled())
+				{
+					Ar.Log(TEXT("Message Bridge: Initialized and enabled"));
+				}
+				else
+				{
+					Ar.Log(TEXT("Message Bridge: Initialized, but disabled"));
+				}
+			}
+			else
+			{
+				Ar.Log(TEXT("Message Bridge: Not initialized."));
+			}
+
+			Ar.Logf(TEXT("  Unicast Endpoint: %s"), *Settings->UnicastEndpoint);
+			Ar.Logf(TEXT("  Multicast Endpoint: %s"), *Settings->MulticastEndpoint);
+			Ar.Logf(TEXT("  Multicast TTL: %i"), Settings->MulticastTimeToLive);
+			Ar.Log(TEXT("  Static Endpoints:"));
+
+			for (const auto& StaticEndpoint : Settings->StaticEndpoints)
+			{
+				Ar.Logf(TEXT("    > %s"), *StaticEndpoint);
+			}
+
+			// tunnel status
+			if (MessageTunnel.IsValid())
+			{
+				if (MessageTunnel->IsServerRunning())
+				{
+					Ar.Log(TEXT("Message Tunnel: Initialized and started"));
+				}
+				else
+				{
+					Ar.Log(TEXT("Message Tunnel: Initialized, but stopped"));
+				}
+			}
+			else
+			{
+				Ar.Log(TEXT("Message Tunnel: Not initialized."));
+			}
+
+			Ar.Logf(TEXT("  Unicast Endpoint: %s"), *Settings->TunnelUnicastEndpoint);
+			Ar.Logf(TEXT("  Multicast Endpoint: %s"), *Settings->TunnelMulticastEndpoint);
+
+			for (const auto& RemoteEndpoint : Settings->RemoteTunnelEndpoints)
+			{
+				Ar.Logf(TEXT("    > %s"), *RemoteEndpoint);
+			}
+		}
+		else if (FParse::Command(&Cmd, TEXT("RESTART")))
+		{
+			RestartServices();
+		}
+		else if (FParse::Command(&Cmd, TEXT("SHUTDOWN")))
+		{
+			ShutdownBridge();
+			ShutdownTunnel();
+		}
+		else
+		{
+			// show usage
+			Ar.Log(TEXT("Usage: UDPMESSAGING <Command>"));
+			Ar.Log(TEXT(""));
+			Ar.Log(TEXT("Command"));
+			Ar.Log(TEXT("    RESTART = Restarts the message bridge and message tunnel, if enabled"));
+			Ar.Log(TEXT("    SHUTDOWN = Shut down the message bridge and message tunnel, if running"));
+			Ar.Log(TEXT("    STATUS = Displays the status of the UDP message transport"));
+		}
+
+		return true;
+	}
+
 public:
 
 	// IModuleInterface interface
diff --git a/Engine/Source/Runtime/Messaging/Private/MessagingModule.cpp b/Engine/Source/Runtime/Messaging/Private/MessagingModule.cpp
--- a/Engine/Source/Runtime/Messaging/Private/MessagingModule.cpp
+++ b/Engine/Source/Runtime/Messaging/Private/MessagingModule.cpp
@@ -23,32 +23,32 @@ public:
 
 	virtual bool Exec( UWorld* InWorld, const TCHAR* Cmd, FOutputDevice& Ar ) override
 	{
-		if (FParse::Command(&Cmd, TEXT("MESSAGING")))
+		if (!FParse::Command(&Cmd, TEXT("MESSAGING")))
 		{
-			if (FParse::Command(&Cmd, TEXT("STATUS")))
+			return false;
+		}
+
+		if (FParse::Command(&Cmd, TEXT("STATUS")))
+		{
+			if (DefaultBus.IsValid())
 			{
-				if (DefaultBus.IsValid())
-				{
-					Ar.Log(TEXT("Default message bus has been initialized."));
-				}
-				else
-				{
-					Ar.Log(TEXT("Default message bus has NOT been initialized yet."));
-				}
+				Ar.Log(TEXT("Default message bus has been initialized."));
 			}
 			else
 			{
-				// show usage
-				Ar.Log(TEXT("Usage: MESSAGING <Command>"));
-				Ar.Log(TEXT(""));
-				Ar.Log(TEXT("Command"));
-				Ar.Log(TEXT("    STATUS = Displays the status of the default message bus"));
+				Ar.Log(TEXT("Default message bus has NOT been initialized yet."));
 			}
-
-			return true;
+		}
+		else
+		{
+			// show usage
+			Ar.Log(TEXT("Usage: MESSAGING <Command>"));
+			Ar.Log(TEXT(""));
+			Ar.Log(TEXT("Command"));
+			Ar.Log(TEXT("    STATUS = Displays the status of the default message bus"));
 		}
 
-		return false;
+		return true;
 	}
 
 public:
