UdpMessaging: Added tunnel connection stats to diagnostics command

diff --git a/Engine/Plugins/Messaging/UdpMessaging/Source/UdpMessaging/Private/UdpMessagingModule.cpp b/Engine/Plugins/Messaging/UdpMessaging/Source/UdpMessaging/Private/UdpMessagingModule.cpp
--- a/Engine/Plugins/Messaging/UdpMessaging/Source/UdpMessaging/Private/UdpMessagingModule.cpp
+++ b/Engine/Plugins/Messaging/UdpMessaging/Source/UdpMessaging/Private/UdpMessagingModule.cpp
@@ -51,14 +51,23 @@ public:
 				Ar.Log(TEXT("Message Bridge: Not initialized."));
 			}
 
+			// bridge settings
 			Ar.Logf(TEXT("  Unicast Endpoint: %s"), *Settings->UnicastEndpoint);
 			Ar.Logf(TEXT("  Multicast Endpoint: %s"), *Settings->MulticastEndpoint);
 			Ar.Logf(TEXT("  Multicast TTL: %i"), Settings->MulticastTimeToLive);
-			Ar.Log(TEXT("  Static Endpoints:"));
 
-			for (const auto& StaticEndpoint : Settings->StaticEndpoints)
+			if (Settings->StaticEndpoints.Num() > 0)
 			{
-				Ar.Logf(TEXT("    > %s"), *StaticEndpoint);
+				Ar.Log(TEXT("  Static Endpoints:"));
+
+				for (const auto& StaticEndpoint : Settings->StaticEndpoints)
+				{
+					Ar.Logf(TEXT("  > %s"), *StaticEndpoint);
+				}
+			}
+			else
+			{
+				Ar.Log(TEXT("  Static Endpoints: None"));
 			}
 
 			// tunnel status
@@ -80,10 +89,39 @@ public:
 
 			Ar.Logf(TEXT("  Unicast Endpoint: %s"), *Settings->TunnelUnicastEndpoint);
 			Ar.Logf(TEXT("  Multicast Endpoint: %s"), *Settings->TunnelMulticastEndpoint);
+			Ar.Log(TEXT("  Remote Endpoints:"));
 
 			for (const auto& RemoteEndpoint : Settings->RemoteTunnelEndpoints)
 			{
-				Ar.Logf(TEXT("    > %s"), *RemoteEndpoint);
+				Ar.Logf(TEXT("  > %s"), *RemoteEndpoint);
+			}
+
+			if (MessageTunnel.IsValid())
+			{
+				Ar.Logf(TEXT("  Total Bytes In: %i"), MessageTunnel->GetTotalInboundBytes());
+				Ar.Logf(TEXT("  Total Bytes Out: %i"), MessageTunnel->GetTotalOutboundBytes());
+
+				TArray<TSharedPtr<IUdpMessageTunnelConnection>> Connections;
+			
+				if (MessageTunnel->GetConnections(Connections) > 0)
+				{
+					Ar.Log(TEXT("  Active Connections:"));
+
+					for (const auto& Connection : Connections)
+					{
+						Ar.Logf(TEXT("  > %s, Open: %s, Uptime: %s, Bytes Received: %i, Bytes Sent: %i"),
+							*Connection->GetName().ToString(),
+							Connection->IsOpen() ? *GYes.ToString() : *GNo.ToString(),
+							*Connection->GetUptime().ToString(),
+							Connection->GetTotalBytesReceived(),
+							Connection->GetTotalBytesSent()
+						);
+					}
+				}
+				else
+				{
+					Ar.Log(TEXT("  Active Connections: None"));
+				}
 			}
 		}
 		else if (FParse::Command(&Cmd, TEXT("RESTART")))
