Fix failure to build lighting on Mac, due to UdpMessaging plugin not loading for UnrealLightmass.

diff --git a/Engine/Plugins/Messaging/UdpMessaging/UdpMessaging.uplugin b/Engine/Plugins/Messaging/UdpMessaging/UdpMessaging.uplugin
--- a/Engine/Plugins/Messaging/UdpMessaging/UdpMessaging.uplugin
+++ b/Engine/Plugins/Messaging/UdpMessaging/UdpMessaging.uplugin
@@ -14,7 +14,7 @@
 	"CanContainContent" : false,
 	"IsBetaVersion" : false,
 	"Installed" : false,
-	"SupportedPrograms" : [ "UnrealFrontend" ],
+	"SupportedPrograms" : [ "UnrealFrontend", "UnrealLightmass" ],
 	"Modules" :
 	[
 		{
@@ -25,7 +25,7 @@
 			[
 				"HTML5"
 			],
-			"WhitelistPrograms" : [ "UnrealFrontend" ]
+			"WhitelistPrograms" : [ "UnrealFrontend", "UnrealLightmass" ]
 		}
 	]
 }
\ No newline at end of file
