Add a local static endpoint to message bus discovery if there isn't already one to properly discover other message bus process bound on the loopback address (i.e. the launcher)

diff --git a/Engine/Plugins/Messaging/UdpMessaging/Source/UdpMessaging/Private/Transport/UdpMessageBeacon.cpp b/Engine/Plugins/Messaging/UdpMessaging/Source/UdpMessaging/Private/Transport/UdpMessageBeacon.cpp
--- a/Engine/Plugins/Messaging/UdpMessaging/Source/UdpMessaging/Private/Transport/UdpMessageBeacon.cpp
+++ b/Engine/Plugins/Messaging/UdpMessaging/Source/UdpMessaging/Private/Transport/UdpMessageBeacon.cpp
@@ -37,6 +37,15 @@ FUdpMessageBeacon::FUdpMessageBeacon(FSocket* InSocket, const FGuid& InSocketId,
 		StaticAddresses.Add(Endpoint.ToInternetAddr());
 	}
 
+	// if we don't have a local endpoint address in our static discovery endpoint, add one.
+	// this is to properly discover other message bus process bound on the loopback address (i.e for the launcher not to trigger firewall)
+	// note: this will only properly work if the socket is not bound (or bound to the any address)
+	FIPv4Endpoint LocalEndpoint(FIPv4Address(127, 0, 0, 1), InMulticastEndpoint.Port);
+	if (!InStaticEndpoints.Contains(LocalEndpoint))
+	{
+		StaticAddresses.Add(LocalEndpoint.ToInternetAddr());
+	}
+
 	Thread = FRunnableThread::Create(this, TEXT("FUdpMessageBeacon"), 128 * 1024, TPri_AboveNormal, FPlatformAffinity::GetPoolThreadMask());
 }
 
