Udp static endpoint configuration only available in Editor and UFE.

diff --git a/Engine/Plugins/Messaging/UdpMessaging/Source/UdpMessaging/Private/Transport/UdpMessageBeacon.cpp b/Engine/Plugins/Messaging/UdpMessaging/Source/UdpMessaging/Private/Transport/UdpMessageBeacon.cpp
--- a/Engine/Plugins/Messaging/UdpMessaging/Source/UdpMessaging/Private/Transport/UdpMessageBeacon.cpp
+++ b/Engine/Plugins/Messaging/UdpMessaging/Source/UdpMessaging/Private/Transport/UdpMessageBeacon.cpp
@@ -10,7 +10,7 @@
 #include "Serialization/ArrayWriter.h"
 #include "Sockets.h"
 
-#if WITH_EDITOR || IS_PROGRAM
+#if WITH_TARGETPLATFORM_SUPPORT
 #include "Misc/CoreMisc.h"
 #include "Interfaces/ITargetPlatformManagerModule.h"
 #include "Interfaces/ITargetPlatform.h"
@@ -51,7 +51,7 @@ FUdpMessageBeacon::FUdpMessageBeacon(FSocket* InSocket, const FGuid& InSocketId,
 		StaticAddresses.Add(LocalEndpoint.ToInternetAddr());
 	}
 
-#if WITH_EDITOR || IS_PROGRAM
+#if WITH_TARGETPLATFORM_SUPPORT
 	if( ITargetPlatformManagerModule* TargetPlatformManager = GetTargetPlatformManager() )
 	{
 		TArray<ITargetPlatform*> Platforms = TargetPlatformManager->GetTargetPlatforms();
@@ -73,7 +73,7 @@ FUdpMessageBeacon::FUdpMessageBeacon(FSocket* InSocket, const FGuid& InSocketId,
 		}
 		ProcessPendingEndpoints();
 	}
-#endif //WITH_EDITOR || IS_PROGRAM
+#endif //WITH_TARGETPLATFORM_SUPPORT
 
 	Thread = FRunnableThread::Create(this, TEXT("FUdpMessageBeacon"), 128 * 1024, TPri_AboveNormal, FPlatformAffinity::GetPoolThreadMask());
 }
@@ -93,7 +93,7 @@ FUdpMessageBeacon::~FUdpMessageBeacon()
 	FPlatformProcess::ReturnSynchEventToPool(EndpointLeftEvent);
 	EndpointLeftEvent = nullptr;
 
-#if WITH_EDITOR || IS_PROGRAM
+#if WITH_TARGETPLATFORM_SUPPORT
 	if( ITargetPlatformManagerModule* TargetPlatformManager = GetTargetPlatformManager() )
 	{
 		TArray<ITargetPlatform*> Platforms = TargetPlatformManager->GetTargetPlatforms();
@@ -107,7 +107,7 @@ FUdpMessageBeacon::~FUdpMessageBeacon()
 
 		}
 	}
-#endif //WITH_EDITOR || IS_PROGRAM
+#endif //WITH_TARGETPLATFORM_SUPPORT
 }
 
 
@@ -243,7 +243,7 @@ bool FUdpMessageBeacon::SendPing(const FTimespan& SocketWaitTime)
 
 void FUdpMessageBeacon::Update(const FDateTime& CurrentTime, const FTimespan& SocketWaitTime)
 {
-#if WITH_EDITOR || IS_PROGRAM
+#if WITH_TARGETPLATFORM_SUPPORT
 	ProcessPendingEndpoints();
 #endif
 
@@ -263,7 +263,7 @@ void FUdpMessageBeacon::Update(const FDateTime& CurrentTime, const FTimespan& So
 }
 
 
-#if WITH_EDITOR || IS_PROGRAM
+#if WITH_TARGETPLATFORM_SUPPORT
 void FUdpMessageBeacon::HandleTargetPlatformDeviceDiscovered( TSharedRef<class ITargetDevice, ESPMode::ThreadSafe> DiscoveredDevice)
 {
 	TSharedPtr<const FInternetAddr> DeviceStaticEndpoint = DiscoveredDevice->GetMessagingEndpoint();
@@ -275,9 +275,9 @@ void FUdpMessageBeacon::HandleTargetPlatformDeviceDiscovered( TSharedRef<class I
 		PendingEndpoints.Enqueue(PendingEndpoint);
 	}
 }
-#endif //WITH_EDITOR || IS_PROGRAM
+#endif //WITH_TARGETPLATFORM_SUPPORT
 
-#if WITH_EDITOR || IS_PROGRAM
+#if WITH_TARGETPLATFORM_SUPPORT
 void FUdpMessageBeacon::HandleTargetPlatformDeviceLost( TSharedRef<class ITargetDevice, ESPMode::ThreadSafe> LostDevice)
 {
 	TSharedPtr<const FInternetAddr> DeviceStaticEndpoint = LostDevice->GetMessagingEndpoint();
@@ -289,9 +289,9 @@ void FUdpMessageBeacon::HandleTargetPlatformDeviceLost( TSharedRef<class ITarget
 		PendingEndpoints.Enqueue(PendingEndpoint);
 	}
 }
-#endif //WITH_EDITOR || IS_PROGRAM
+#endif //WITH_TARGETPLATFORM_SUPPORT
 
-#if WITH_EDITOR || IS_PROGRAM
+#if WITH_TARGETPLATFORM_SUPPORT
 void FUdpMessageBeacon::ProcessPendingEndpoints()
 {
 	// process any pending static endpoint changes from target platform device changes
@@ -316,7 +316,7 @@ void FUdpMessageBeacon::ProcessPendingEndpoints()
 		}
 	}
 }
-#endif //WITH_EDITOR || IS_PROGRAM
+#endif //WITH_TARGETPLATFORM_SUPPORT
 
 
 
diff --git a/Engine/Plugins/Messaging/UdpMessaging/Source/UdpMessaging/Private/Transport/UdpMessageBeacon.h b/Engine/Plugins/Messaging/UdpMessaging/Source/UdpMessaging/Private/Transport/UdpMessageBeacon.h
--- a/Engine/Plugins/Messaging/UdpMessaging/Source/UdpMessaging/Private/Transport/UdpMessageBeacon.h
+++ b/Engine/Plugins/Messaging/UdpMessaging/Source/UdpMessaging/Private/Transport/UdpMessageBeacon.h
@@ -17,6 +17,10 @@ class FInternetAddr;
 class FSocket;
 struct FIPv4Endpoint;
 
+#if !defined(WITH_TARGETPLATFORM_SUPPORT)
+	#define WITH_TARGETPLATFORM_SUPPORT 0
+#endif
+
 
 /**
  * Implements a beacon sender thread.
@@ -104,11 +108,11 @@ protected:
 
 private:
 
-#if WITH_EDITOR || IS_PROGRAM
+#if WITH_TARGETPLATFORM_SUPPORT
 	void HandleTargetPlatformDeviceDiscovered( TSharedRef<class ITargetDevice, ESPMode::ThreadSafe> DiscoveredDevice );
 	void HandleTargetPlatformDeviceLost( TSharedRef<class ITargetDevice, ESPMode::ThreadSafe> LostDevice );
 	void ProcessPendingEndpoints();
-#endif //WITH_EDITOR || IS_PROGRAM
+#endif //WITH_TARGETPLATFORM_SUPPORT
 
 
 	/** Holds the calculated interval between Hello segments. */
@@ -144,7 +148,7 @@ private:
 	/** Holds the thread object. */
 	FRunnableThread* Thread;
 
-#if WITH_EDITOR || IS_PROGRAM
+#if WITH_TARGETPLATFORM_SUPPORT
 	/** Holds target devices that have just been discovered **/
 	struct FPendingEndpoint
 	{
@@ -152,7 +156,7 @@ private:
 		bool bAdd;
 	};
 	TQueue<FPendingEndpoint,EQueueMode::Mpsc> PendingEndpoints;
-#endif //WITH_EDITOR || IS_PROGRAM
+#endif //WITH_TARGETPLATFORM_SUPPORT
 
 private:
 	
diff --git a/Engine/Plugins/Messaging/UdpMessaging/Source/UdpMessaging/UdpMessaging.Build.cs b/Engine/Plugins/Messaging/UdpMessaging/Source/UdpMessaging/UdpMessaging.Build.cs
--- a/Engine/Plugins/Messaging/UdpMessaging/Source/UdpMessaging/UdpMessaging.Build.cs
+++ b/Engine/Plugins/Messaging/UdpMessaging/Source/UdpMessaging/UdpMessaging.Build.cs
@@ -41,20 +41,39 @@ namespace UnrealBuildTool.Rules
 					"UdpMessaging/Private/Tunnel",
 				});
 
-			if (Target.Type == TargetType.Editor || Target.Type == TargetType.Program)
+			if (Target.Type == TargetType.Editor)
 			{
 				DynamicallyLoadedModuleNames.AddRange(
 					new string[] {
 						"Settings",
-						"TargetPlatform",
 					});
 
 				PrivateIncludePathModuleNames.AddRange(
 					new string[] {
 						"Settings",
+					});
+			}
+
+			//temporary workaround to allow automatic configuration of static endpoints in editor/UFE for certain target devices, pending proper API
+			if (Target.Type == TargetType.Editor || (Target.Type == TargetType.Program && Target.Name == "UnrealFrontend"))
+			{
+				PublicDefinitions.Add("WITH_TARGETPLATFORM_SUPPORT=1");
+
+				DynamicallyLoadedModuleNames.AddRange(
+					new string[] {
 						"TargetPlatform",
 					});
+
+				PrivateIncludePathModuleNames.AddRange(
+					new string[] {
+						"TargetPlatform",
+					});
+			}
+			else
+			{
+				PublicDefinitions.Add("WITH_TARGETPLATFORM_SUPPORT=0");
 			}
+
 		}
 	}
 }
