Copied 4372687 Francis.Hurteau
Fix potential assert when changing ProtocolVersion based on ping messages

diff --git a/Engine/Plugins/Messaging/UdpMessaging/Source/UdpMessaging/Private/Transport/UdpMessageProcessor.cpp b/Engine/Plugins/Messaging/UdpMessaging/Source/UdpMessaging/Private/Transport/UdpMessageProcessor.cpp
--- a/Engine/Plugins/Messaging/UdpMessaging/Source/UdpMessaging/Private/Transport/UdpMessageProcessor.cpp
+++ b/Engine/Plugins/Messaging/UdpMessaging/Source/UdpMessaging/Private/Transport/UdpMessageProcessor.cpp
@@ -228,9 +228,9 @@ void FUdpMessageProcessor::AcknowledgeReceipt(int32 MessageId, const FNodeInfo&
 {
 	FUdpMessageSegment::FHeader Header;
 	{
+		Header.ProtocolVersion = NodeInfo.ProtocolVersion;
 		Header.RecipientNodeId = NodeInfo.NodeId;
 		Header.SenderNodeId = LocalNodeId;
-		Header.ProtocolVersion = NodeInfo.ProtocolVersion;
 		Header.SegmentType = EUdpMessageSegments::Acknowledge;
 	}
 
@@ -262,14 +262,13 @@ void FUdpMessageProcessor::ConsumeInboundSegments()
 
 	while (InboundSegments.Dequeue(Segment))
 	{
-		FUdpMessageSegment::FHeader Header;
-
 		// quick hack for TTP# 247103
 		if (!Segment.Data.IsValid())
 		{
 			continue;
 		}
 
+		FUdpMessageSegment::FHeader Header;
 		*Segment.Data << Header;
 
 		if (FilterSegment(Header))
@@ -283,8 +282,9 @@ void FUdpMessageProcessor::ConsumeInboundSegments()
 				NodeDiscoveredDelegate.ExecuteIfBound(NodeInfo.NodeId);
 			}
 
-			check(NodeInfo.ProtocolVersion == Header.ProtocolVersion);
+			NodeInfo.ProtocolVersion = Header.ProtocolVersion;
 			NodeInfo.Endpoint = Segment.Sender;
+			NodeInfo.LastSegmentReceivedTime = CurrentTime;
 
 			switch (Header.SegmentType)
 			{
@@ -331,8 +331,6 @@ void FUdpMessageProcessor::ConsumeInboundSegments()
 			default:
 				ProcessUnknownSegment(Segment, NodeInfo, (uint8)Header.SegmentType);
 			}
-
-			NodeInfo.LastSegmentReceivedTime = CurrentTime;
 		}
 	}
 }
@@ -509,7 +507,7 @@ void FUdpMessageProcessor::ProcessPingSegment(FInboundSegment& Segment, FNodeInf
 	// if that protocol isn't in our supported protocols we do not reply to the pong and remove this node since we don't support its version
 	if (!SupportedProtocolVersions.Contains(ProtocolVersion))
 	{
-		KnownNodes.Remove(NodeInfo.NodeId);
+		RemoveKnownNode(NodeInfo.NodeId);
 		return;
 	}
 
@@ -519,10 +517,10 @@ void FUdpMessageProcessor::ProcessPingSegment(FInboundSegment& Segment, FNodeInf
 	// Send the pong
 	FUdpMessageSegment::FHeader Header;
 	{
-		Header.RecipientNodeId = NodeInfo.NodeId;
-		Header.SenderNodeId = LocalNodeId;
 		// Reply to the ping using the agreed protocol
 		Header.ProtocolVersion = ProtocolVersion;
+		Header.RecipientNodeId = NodeInfo.NodeId;
+		Header.SenderNodeId = LocalNodeId;
 		Header.SegmentType = EUdpMessageSegments::Pong;
 	}
 
