Fix indefinite wait in beacon if wait time ends up being negative for some reason

diff --git a/Engine/Plugins/Messaging/UdpMessaging/Source/UdpMessaging/Private/Transport/UdpMessageBeacon.cpp b/Engine/Plugins/Messaging/UdpMessaging/Source/UdpMessaging/Private/Transport/UdpMessageBeacon.cpp
--- a/Engine/Plugins/Messaging/UdpMessaging/Source/UdpMessaging/Private/Transport/UdpMessageBeacon.cpp
+++ b/Engine/Plugins/Messaging/UdpMessaging/Source/UdpMessaging/Private/Transport/UdpMessageBeacon.cpp
@@ -99,7 +99,9 @@ uint32 FUdpMessageBeacon::Run()
 	{
 		FDateTime CurrentTime = FDateTime::UtcNow();
 		Update(CurrentTime, BeaconInterval);
-		EndpointLeftEvent->Wait(NextHelloTime - CurrentTime);
+		// Clamp the wait time to a positive value
+		uint32 WaitTimeMs = (uint32)FMath::Max((NextHelloTime - CurrentTime).GetTotalMilliseconds(), 0.0);
+		EndpointLeftEvent->Wait(WaitTimeMs);
 	}
 
 	SendSegment(EUdpMessageSegments::Bye, BeaconInterval);
