Fix crash in UDP Messaging tunnel when failing to bind unicast endpoint

diff --git a/Engine/Plugins/Messaging/UdpMessaging/Source/UdpMessaging/Private/Tunnel/UdpMessageTunnel.cpp b/Engine/Plugins/Messaging/UdpMessaging/Source/UdpMessaging/Private/Tunnel/UdpMessageTunnel.cpp
--- a/Engine/Plugins/Messaging/UdpMessaging/Source/UdpMessaging/Private/Tunnel/UdpMessageTunnel.cpp
+++ b/Engine/Plugins/Messaging/UdpMessaging/Source/UdpMessaging/Private/Tunnel/UdpMessageTunnel.cpp
@@ -24,9 +24,21 @@ FUdpMessageTunnel::FUdpMessageTunnel(const FIPv4Endpoint& InUnicastEndpoint, con
 	: Listener(nullptr)
 	, MulticastEndpoint(InMulticastEndpoint)
 	, Stopping(false)
+	, Thread(nullptr)
 	, TotalInboundBytes(0)
 	, TotalOutboundBytes(0)
 {
+	UnicastSocket = FUdpSocketBuilder(TEXT("UdpMessageUnicastSocket"))
+		.AsNonBlocking()
+		.BoundToEndpoint(InUnicastEndpoint)
+		.WithReceiveBufferSize(UDP_MESSAGING_RECEIVE_BUFFER_SIZE);
+
+	if (UnicastSocket == nullptr)
+	{
+		UE_LOG(LogUdpMessaging, Error, TEXT("UdpMessageTunnel failed to create unicast socket on %s"), *InUnicastEndpoint.ToString());
+		Stopping = true;
+	}
+
 	// initialize sockets
 	MulticastSocket = FUdpSocketBuilder(TEXT("UdpMessageMulticastSocket"))
 		.AsNonBlocking()
@@ -39,21 +51,23 @@ FUdpMessageTunnel::FUdpMessageTunnel(const FIPv4Endpoint& InUnicastEndpoint, con
 		// interface IP instead of the multicast address.
 		.BoundToAddress(InUnicastEndpoint.Address)
 #endif
-		.BoundToPort(MulticastEndpoint.Port)
 		.BoundToPort(InMulticastEndpoint.Port)
 		.JoinedToGroup(InMulticastEndpoint.Address)
 		.WithMulticastLoopback()
-		.WithMulticastTtl(1);
+		.WithMulticastTtl(1)
+		.WithReceiveBufferSize(UDP_MESSAGING_RECEIVE_BUFFER_SIZE);
 
-	UnicastSocket = FUdpSocketBuilder(TEXT("UdpMessageUnicastSocket"))
-		.AsNonBlocking()
-		.BoundToEndpoint(InUnicastEndpoint);
-
-	int32 NewSize = 0;
-	MulticastSocket->SetReceiveBufferSize(2 * 1024 * 1024, NewSize);
-	UnicastSocket->SetReceiveBufferSize(2 * 1024 * 1024, NewSize);
+	if (MulticastSocket == nullptr)
+	{
+		UE_LOG(LogUdpMessaging, Error, TEXT("UdpMessageTunnel failed to create multicast socket on %s, joined to %s"), *InUnicastEndpoint.ToString(), *InMulticastEndpoint.ToString());
+		Stopping = true;
+	}
 
-	Thread = FRunnableThread::Create(this, TEXT("FUdpMessageTunnel"), 128 * 1024, TPri_AboveNormal);
+	// Start the tunnel only if properly initialized
+	if (!Stopping)
+	{
+		Thread = FRunnableThread::Create(this, TEXT("FUdpMessageTunnel"), 128 * 1024, TPri_AboveNormal);
+	}
 }
 
 
