Added command line argument to override udpmessaging transport unicast and multicast endpoint
Merged from cl 7280851

diff --git a/Engine/Plugins/Messaging/UdpMessaging/Source/UdpMessaging/Private/Shared/UdpMessagingSettings.h b/Engine/Plugins/Messaging/UdpMessaging/Source/UdpMessaging/Private/Shared/UdpMessagingSettings.h
--- a/Engine/Plugins/Messaging/UdpMessaging/Source/UdpMessaging/Private/Shared/UdpMessagingSettings.h
+++ b/Engine/Plugins/Messaging/UdpMessaging/Source/UdpMessaging/Private/Shared/UdpMessagingSettings.h
@@ -19,7 +19,10 @@ class UUdpMessagingSettings
 
 public:
 
-	/** Whether the UDP transport channel is enabled. */
+	/**
+	 * Whether the UDP transport channel is enabled.
+	 * Can be specified on the command line with `-UDPMESSAGING_TRANSPORT_ENABLE=`
+	 */
 	UPROPERTY(config, EditAnywhere, Category=Transport)
 	bool EnableTransport;
 
@@ -29,6 +32,7 @@ public:
 	 * The format is IP_ADDRESS:PORT_NUMBER.
 	 * 0.0.0.0:0 will bind to the default network adapter on Windows,
 	 * and all available network adapters on other operating systems.
+	 * Can be specified on the command line with `-UDPMESSAGING_TRANSPORT_UNICAST=`
 	 */
 	UPROPERTY(config, EditAnywhere, Category=Transport)
 	FString UnicastEndpoint;
@@ -38,6 +42,7 @@ public:
 	 *
 	 * The format is IP_ADDRESS:PORT_NUMBER.
 	 * The multicast IP address must be in the range 224.0.0.0 to 239.255.255.255.
+	 * Can be specified on the command line with `-UDPMESSAGING_TRANSPORT_MULTICAST=`
 	 */
 	UPROPERTY(config, EditAnywhere, Category=Transport)
 	FString MulticastEndpoint;
diff --git a/Engine/Plugins/Messaging/UdpMessaging/Source/UdpMessaging/Private/UdpMessagingModule.cpp b/Engine/Plugins/Messaging/UdpMessaging/Source/UdpMessaging/Private/UdpMessagingModule.cpp
--- a/Engine/Plugins/Messaging/UdpMessaging/Source/UdpMessaging/Private/UdpMessagingModule.cpp
+++ b/Engine/Plugins/Messaging/UdpMessaging/Source/UdpMessaging/Private/UdpMessagingModule.cpp
@@ -249,6 +249,9 @@ public:
 		}
 #endif // WITH_EDITOR
 
+		// parse additional command line args
+		ParseCommandLine(GetMutableDefault<UUdpMessagingSettings>(), FCommandLine::Get());
+
 		// register application events
 		FCoreDelegates::ApplicationHasReactivatedDelegate.AddRaw(this, &FUdpMessagingModule::HandleApplicationHasReactivated);
 		FCoreDelegates::ApplicationWillDeactivateDelegate.AddRaw(this, &FUdpMessagingModule::HandleApplicationWillDeactivate);
@@ -426,6 +429,20 @@ protected:
 #endif
 	}
 
+	/**
+	 * Parse command line arguments to override UdpMessagingSettings
+	 */
+	void ParseCommandLine(UUdpMessagingSettings* Settings, const TCHAR* CommandLine)
+	{
+		if (Settings && CommandLine)
+		{
+			// Parse value overrides (if present)
+			FParse::Bool(CommandLine, TEXT("-UDPMESSAGING_ENABLE_TRANSPORT="), Settings->EnableTransport);
+			FParse::Value(CommandLine, TEXT("-UDPMESSAGING_TRANSPORT_UNICAST="), Settings->UnicastEndpoint);
+			FParse::Value(CommandLine, TEXT("-UDPMESSAGING_TRANSPORT_MULTICAST="), Settings->MulticastEndpoint);
+		}
+	}
+
 	/**
 	 * Checks whether networked message transport is supported.
 	 *
