- On AppPreExit, stopped the UdpTransport and waited for all serialization tasks fired by the UdpMessageProcessor to complete.

Note: This bug didn't only occur on Mac or DisasterRecovery. It could occur on all platforms and for any system using UdpMessaging and sending large messages just before the engine shuts down.

diff --git a/Engine/Plugins/Messaging/UdpMessaging/Source/UdpMessaging/Private/Transport/UdpMessageProcessor.cpp b/Engine/Plugins/Messaging/UdpMessaging/Source/UdpMessaging/Private/Transport/UdpMessageProcessor.cpp
--- a/Engine/Plugins/Messaging/UdpMessaging/Source/UdpMessaging/Private/Transport/UdpMessageProcessor.cpp
+++ b/Engine/Plugins/Messaging/UdpMessaging/Source/UdpMessaging/Private/Transport/UdpMessageProcessor.cpp
@@ -120,6 +120,11 @@ TMap<uint8, TArray<FGuid>> FUdpMessageProcessor::GetRecipientsPerProtocolVersion
 
 bool FUdpMessageProcessor::EnqueueInboundSegment(const TSharedPtr<FArrayReader, ESPMode::ThreadSafe>& Data, const FIPv4Endpoint& InSender)
 {
+	if (Stopping)
+	{
+		return false;
+	}
+
 	if (!InboundSegments.Enqueue(FInboundSegment(Data, InSender)))
 	{
 		return false;
@@ -132,6 +137,11 @@ bool FUdpMessageProcessor::EnqueueInboundSegment(const TSharedPtr<FArrayReader,
 
 bool FUdpMessageProcessor::EnqueueOutboundMessage(const TSharedRef<IMessageContext, ESPMode::ThreadSafe>& MessageContext, const TArray<FGuid>& Recipients)
 {
+	if (Stopping)
+	{
+		return false;
+	}
+	
 	TMap<uint8, TArray<FGuid>> RecipientPerVersions = GetRecipientsPerProtocolVersion(Recipients);
 	for (const auto& RecipientVersion : RecipientPerVersions)
 	{
@@ -192,7 +202,7 @@ uint32 FUdpMessageProcessor::Run()
 		UpdateKnownNodes();
 		UpdateStaticNodes();
 	}
-	
+
 	delete Beacon;
 	Beacon = nullptr;
 
@@ -210,6 +220,44 @@ void FUdpMessageProcessor::Stop()
 }
 
 
+void FUdpMessageProcessor::WaitAsyncTaskCompletion()
+{
+	// Check if processor has in-flight serialization task(s).
+	auto HasIncompleteSerializationTasks = [this]()
+	{
+		for (const TPair<FGuid, FNodeInfo>& GuidNodeInfoPair : KnownNodes)
+		{
+			for (const TPair<int32, TSharedPtr<FUdpMessageSegmenter>>& SegmenterPair: GuidNodeInfoPair.Value.Segmenters)
+			{
+				if (!SegmenterPair.Value->IsMessageSerializationDone())
+				{
+					return true;
+				}
+			}
+		}
+
+		for (const TPair<FIPv4Endpoint, FNodeInfo>& Ipv4NodeInfoPair : StaticNodes)
+		{
+			for (const TPair<int32, TSharedPtr<FUdpMessageSegmenter>>& SegmenterPair: Ipv4NodeInfoPair.Value.Segmenters)
+			{
+				if (!SegmenterPair.Value->IsMessageSerializationDone())
+				{
+					return true;
+				}
+			}
+		}
+
+		return false;
+	};
+
+	// Ensures the task graph doesn't contain any pending/running serialization tasks after the processor exit. If the engine is shutting down, the serialization (UStruct) might
+	// not be available anymore when the task is run (The task graph shuts down after the UStruct stuff).
+	while (HasIncompleteSerializationTasks())
+	{
+		FPlatformProcess::Sleep(0); // Yield.
+	}
+}
+
 /* FSingleThreadRunnable interface
 *****************************************************************************/
 
diff --git a/Engine/Plugins/Messaging/UdpMessaging/Source/UdpMessaging/Private/Transport/UdpMessageProcessor.h b/Engine/Plugins/Messaging/UdpMessaging/Source/UdpMessaging/Private/Transport/UdpMessageProcessor.h
--- a/Engine/Plugins/Messaging/UdpMessaging/Source/UdpMessaging/Private/Transport/UdpMessageProcessor.h
+++ b/Engine/Plugins/Messaging/UdpMessaging/Source/UdpMessaging/Private/Transport/UdpMessageProcessor.h
@@ -174,6 +174,12 @@ public:
 		return WorkEvent;
 	}
 
+	/**
+	 * Waits for all serialization tasks fired by this processor to complete. Expected to be called when the application exit
+	 * to prevent serialized (UStruct) object to being use after the UObject system is shutdown.
+	 */
+	void WaitAsyncTaskCompletion();
+
 public:
 
 	// @todo gmp: remove the need for this typedef
diff --git a/Engine/Plugins/Messaging/UdpMessaging/Source/UdpMessaging/Private/Transport/UdpMessageSegmenter.cpp b/Engine/Plugins/Messaging/UdpMessaging/Source/UdpMessaging/Private/Transport/UdpMessageSegmenter.cpp
--- a/Engine/Plugins/Messaging/UdpMessaging/Source/UdpMessaging/Private/Transport/UdpMessageSegmenter.cpp
+++ b/Engine/Plugins/Messaging/UdpMessaging/Source/UdpMessaging/Private/Transport/UdpMessageSegmenter.cpp
@@ -114,6 +114,12 @@ void FUdpMessageSegmenter::Initialize()
 }
 
 
+bool FUdpMessageSegmenter::IsMessageSerializationDone() const
+{
+	return SerializedMessage == nullptr || SerializedMessage->GetState() != EUdpSerializedMessageState::Incomplete;
+}
+
+
 bool FUdpMessageSegmenter::IsInvalid() const
 {
 	return (SerializedMessage->GetState() == EUdpSerializedMessageState::Invalid);
diff --git a/Engine/Plugins/Messaging/UdpMessaging/Source/UdpMessaging/Private/Transport/UdpMessageSegmenter.h b/Engine/Plugins/Messaging/UdpMessaging/Source/UdpMessaging/Private/Transport/UdpMessageSegmenter.h
--- a/Engine/Plugins/Messaging/UdpMessaging/Source/UdpMessaging/Private/Transport/UdpMessageSegmenter.h
+++ b/Engine/Plugins/Messaging/UdpMessaging/Source/UdpMessaging/Private/Transport/UdpMessageSegmenter.h
@@ -177,6 +177,16 @@ public:
 	 */
 	void UpdateSentTime(const FDateTime& CurrentTime);
 
+	/**
+	 * Checks whether the serialization of the message to segment and send completed. The serialization may be performing asynchronously and can may succeed or fail. The message cannot be segmented
+	 * and sent until it was serialized sucessfully.
+	 * @return true if the message is serialization is done, false if it still serializing.
+	 * @see Initialize()
+	 * @see IsInvalid()
+	 * @see IsInitialized()
+	 */
+	bool IsMessageSerializationDone() const;
+
 private:
 	/** Defines the time interval for sending. */
 	static const FTimespan SendInterval;
diff --git a/Engine/Plugins/Messaging/UdpMessaging/Source/UdpMessaging/Private/Transport/UdpMessageTransport.cpp b/Engine/Plugins/Messaging/UdpMessaging/Source/UdpMessaging/Private/Transport/UdpMessageTransport.cpp
--- a/Engine/Plugins/Messaging/UdpMessaging/Source/UdpMessaging/Private/Transport/UdpMessageTransport.cpp
+++ b/Engine/Plugins/Messaging/UdpMessaging/Source/UdpMessaging/Private/Transport/UdpMessageTransport.cpp
@@ -44,6 +44,14 @@ FUdpMessageTransport::~FUdpMessageTransport()
 	StopTransport();
 }
 
+void FUdpMessageTransport::OnAppPreExit()
+{
+	if (MessageProcessor)
+	{
+		MessageProcessor->Stop();
+		MessageProcessor->WaitAsyncTaskCompletion();
+	}
+}
 
 /* IMessageTransport interface
  *****************************************************************************/
diff --git a/Engine/Plugins/Messaging/UdpMessaging/Source/UdpMessaging/Private/Transport/UdpMessageTransport.h b/Engine/Plugins/Messaging/UdpMessaging/Source/UdpMessaging/Private/Transport/UdpMessageTransport.h
--- a/Engine/Plugins/Messaging/UdpMessaging/Source/UdpMessaging/Private/Transport/UdpMessageTransport.h
+++ b/Engine/Plugins/Messaging/UdpMessaging/Source/UdpMessaging/Private/Transport/UdpMessageTransport.h
@@ -46,6 +46,9 @@ public:
 	/** Virtual destructor. */
 	virtual ~FUdpMessageTransport();
 
+	/** Notifies the transport that the application is about to exit. */
+	void OnAppPreExit();
+
 public:
 
 	//~ IMessageTransport interface
diff --git a/Engine/Plugins/Messaging/UdpMessaging/Source/UdpMessaging/Private/UdpMessagingModule.cpp b/Engine/Plugins/Messaging/UdpMessaging/Source/UdpMessaging/Private/UdpMessagingModule.cpp
--- a/Engine/Plugins/Messaging/UdpMessaging/Source/UdpMessaging/Private/UdpMessagingModule.cpp
+++ b/Engine/Plugins/Messaging/UdpMessaging/Source/UdpMessaging/Private/UdpMessagingModule.cpp
@@ -230,6 +230,9 @@ public:
 			return;
 		}
 
+		// Hook to the PreExit callback, needed to execute UObject related shutdowns
+		FCoreDelegates::OnPreExit.AddRaw(this, &FUdpMessagingModule::HandleAppPreExit);
+
 #if WITH_EDITOR
 		// register settings
 		ISettingsModule* SettingsModule = FModuleManager::GetModulePtr<ISettingsModule>("Settings");
@@ -265,6 +268,10 @@ public:
 		FCoreDelegates::ApplicationHasReactivatedDelegate.RemoveAll(this);
 		FCoreDelegates::ApplicationWillDeactivateDelegate.RemoveAll(this);
 
+		// Unhook AppPreExit and call it
+		FCoreDelegates::OnPreExit.RemoveAll(this);
+		HandleAppPreExit();
+
 #if WITH_EDITOR
 		// unregister settings
 		ISettingsModule* SettingsModule = FModuleManager::GetModulePtr<ISettingsModule>("Settings");
@@ -337,9 +344,11 @@ protected:
 
 		UE_LOG(LogUdpMessaging, Log, TEXT("Initializing bridge on interface %s to multicast group %s."), *UnicastEndpoint.ToString(), *MulticastEndpoint.ToText().ToString());
 
+		TSharedRef<FUdpMessageTransport, ESPMode::ThreadSafe> Transport = MakeShared<FUdpMessageTransport, ESPMode::ThreadSafe>(UnicastEndpoint, MulticastEndpoint, Settings->MulticastTimeToLive);
+		WeakBridgeTransport = Transport;
 		MessageBridge = FMessageBridgeBuilder()
-			.UsingTransport(MakeShareable(new FUdpMessageTransport(UnicastEndpoint, MulticastEndpoint, Settings->MulticastTimeToLive)));
-	}
+			.UsingTransport(Transport);
+}
 
 #if PLATFORM_DESKTOP
 	/** Initializes the message tunnel with the current settings. */
@@ -518,11 +527,22 @@ private:
 		return true;
 	}
 
+	void HandleAppPreExit()
+	{
+		if (TSharedPtr<FUdpMessageTransport, ESPMode::ThreadSafe> UdpTransport = WeakBridgeTransport.Pin())
+		{
+			UdpTransport->OnAppPreExit();
+		}
+	}
+
 private:
 
 	/** Holds the message bridge if present. */
 	TSharedPtr<IMessageBridge, ESPMode::ThreadSafe> MessageBridge;
 
+	/** Holds the bridge transport if present.  */
+	TWeakPtr<FUdpMessageTransport, ESPMode::ThreadSafe> WeakBridgeTransport;
+
 #if PLATFORM_DESKTOP
 	/** Holds the message tunnel if present. */
 	TSharedPtr<IUdpMessageTunnel> MessageTunnel;
