Fix UDPMessaging to reject invalid message segments.

diff --git a/Engine/Plugins/Messaging/UdpMessaging/Source/UdpMessaging/Private/Transport/UdpMessageProcessor.cpp b/Engine/Plugins/Messaging/UdpMessaging/Source/UdpMessaging/Private/Transport/UdpMessageProcessor.cpp
--- a/Engine/Plugins/Messaging/UdpMessaging/Source/UdpMessaging/Private/Transport/UdpMessageProcessor.cpp
+++ b/Engine/Plugins/Messaging/UdpMessaging/Source/UdpMessaging/Private/Transport/UdpMessageProcessor.cpp
@@ -480,19 +480,64 @@ void FUdpMessageProcessor::ProcessDataSegment(FInboundSegment& Segment, FNodeInf
 	FUdpMessageSegment::FDataChunk DataChunk;
 	DataChunk.Serialize(*Segment.Data, NodeInfo.ProtocolVersion);
 	
+	if (Segment.Data->IsError())
+	{
+		UE_LOG(LogUdpMessaging, Warning, TEXT("FUdpMessageProcessor::ProcessDataSegment: Failed to serialize DataChunk. Sender=%s"),
+			*(Segment.Sender.ToString()));
+		return;
+	}
+
 	// Discard late segments for sequenced messages
 	if ((DataChunk.Sequence != 0) && (DataChunk.Sequence < NodeInfo.Resequencer.GetNextSequence()))
 	{
 		return;
 	}
-
 	TSharedPtr<FUdpReassembledMessage, ESPMode::ThreadSafe>& ReassembledMessage = NodeInfo.ReassembledMessages.FindOrAdd(DataChunk.MessageId);
 
 	// Reassemble message
 	if (!ReassembledMessage.IsValid())
 	{
 		ReassembledMessage = MakeShared<FUdpReassembledMessage, ESPMode::ThreadSafe>(NodeInfo.ProtocolVersion, DataChunk.MessageFlags, DataChunk.MessageSize, DataChunk.TotalSegments, DataChunk.Sequence, Segment.Sender);
+
+		if (ReassembledMessage->IsMalformed())
+		{
+			// Go ahead and throw away the message.
+			// The sender should see the NAK and resend, so we'll attempt to recreate it later.
+			UE_LOG(LogUdpMessaging, Warning, TEXT("FUdpMessageProcessor::ProcessDataSegment: Ignoring malformed Message %s"), *(ReassembledMessage->Describe()));
+			NodeInfo.ReassembledMessages.Remove(DataChunk.MessageId);
+			ReassembledMessage.Reset();
+			return;
+		}
+	}
+
+	/**
+	// TODO: In a future release uncomment these checks.
+	//		Don't do this for 4.23, because there could be existing third party tools / producers that
+	//		just send dummy data (they shouldn't!), and this could break them unexpectedly.
+	//		These should probably also be moved into a shared location (like into FUdpReassembledMessage).
+
+	if (ReassembledMessage->GetTotalSegmentsCount() != DataChunk.TotalSegments)
+	{
+		UE_LOG(LogUdpMessaging, Warning, TEXT("FUdpMessageProcessor::ProcessDataSegment: Ignoring segment with invalid TotalSegment count. Message=%s, ExpectedTotalSegments=%lu, InTotalSegments=%lu"),
+			*ReassembledMessage->Describe(), ReassembledMessage->GetTotalSegmentCount(), DataChunk.TotalSegments);
+		return;
+	}
+	if (ReassembledMessage->GetData().Num() != DataChunk.MessageSize)
+	{
+		UE_LOG(LogUdpMessaging, Warning, TEXT("FUdpMessageProcessor::ProcessDataSegment: Ignoring segment with invalid MessageSize. Message=%s, ExpectedMessageSize=%d, InMessageSize=%d"),
+			*ReassembledMessage->Describe(), ReassembledMessage->GetData().Num(), DataChunk.MessageSize);
+		return;
 	}
+	if (ReassembledMessage->GetSequence() != DataChunk.Sequence)
+	{
+		UE_LOG(LogUdpMessaging, Warning, TEXT("FUdpMessageProcessor::ProcessDataSegment: Ignoring segment with invalid Sequence. Message=%s, ExpectedSequence=%llu, InSequence=%llu"),
+			*ReassembledMessage->Describe(), ReassembledMessaged->GetSequence(), DataChunk.Sequence);
+		return;
+	}
+
+	// TODO: Check MessageFlags.
+	// TODO: Check Sender.
+	*/
 
 	ReassembledMessage->Reassemble(DataChunk.SegmentNumber, DataChunk.SegmentOffset, DataChunk.Data, CurrentTime);
 
diff --git a/Engine/Plugins/Messaging/UdpMessaging/Source/UdpMessaging/Private/Transport/UdpReassembledMessage.h b/Engine/Plugins/Messaging/UdpMessaging/Source/UdpMessaging/Private/Transport/UdpReassembledMessage.h
--- a/Engine/Plugins/Messaging/UdpMessaging/Source/UdpMessaging/Private/Transport/UdpReassembledMessage.h
+++ b/Engine/Plugins/Messaging/UdpMessaging/Source/UdpMessaging/Private/Transport/UdpReassembledMessage.h
@@ -8,6 +8,7 @@
 #include "Interfaces/IPv4/IPv4Endpoint.h"
 #include "Misc/DateTime.h"
 #include "UdpMessagingPrivate.h"
+#include "Math/NumericLimits.h"
 
 // IMessageContext forward declaration
 enum class EMessageFlags : uint32;
@@ -35,7 +36,6 @@ public:
 	FUdpReassembledMessage(uint8 InProtocolVersion, EMessageFlags InFlags, int64 MessageSize, uint32 SegmentCount, uint64 InSequence, const FIPv4Endpoint& InSender)
 		: ProtocolVersion(InProtocolVersion)
 		, MessageFlags(InFlags)
-		, PendingSegments(true, SegmentCount)
 		, PendingSegmentsCount(SegmentCount)
 		, ReceivedBytes(0)
 		, bIsDelivered(false)
@@ -43,7 +43,29 @@ public:
 		, Sender(InSender)
 		, Sequence(InSequence)
 	{
-		Data.AddUninitialized(MessageSize);
+		// Containers only use int32 for indices, so we can't actually allow sizes greater than that
+		// without causing overflow issues.
+		if (MessageSize > static_cast<int64>(TNumericLimits<int32>::Max()) || MessageSize <= 0)
+		{
+			UE_LOG(LogUdpMessaging, Warning, TEXT("FUdpReassembledMessage: Invalid MessageSize, marking message malformed. Message=%s, MessageSize=%lli"),
+				*Describe(), MessageSize);
+			bIsMalformed = true;
+		}
+		if (SegmentCount > static_cast<uint32>(TNumericLimits<int32>::Max()) || SegmentCount == 0)
+		{
+			UE_LOG(LogUdpMessaging, Warning, TEXT("FUdpReassembledMessage: Invalid SegmentCount, marking message malformed. Message=%s, SegmentCount=%u"),
+				*Describe(), SegmentCount);
+			bIsMalformed = true;
+		}
+
+		// It's completely possible that MessageSize and SegmentCount are still incorrect due to corruption.
+		// However, we have no way to detect that at this point, so we'll assume they're fine and validate
+		// as we receive segments later.
+		if (!IsMalformed())
+		{
+			PendingSegments = TBitArray<>(true, SegmentCount);
+			Data.AddUninitialized(MessageSize);
+		}
 	}
 
 	~FUdpReassembledMessage() = default;
@@ -168,7 +190,7 @@ public:
 	 */
 	bool IsInitialized() const
 	{
-		return (Data.Num() < 0);
+		return (Data.Num() > 0);
 	}
 
 	/**
@@ -230,29 +252,54 @@ public:
 	 */
 	void Reassemble(int32 SegmentNumber, int32 SegmentOffset, const TArray<uint8>& SegmentData, const FDateTime& CurrentTime)
 	{
-		if (SegmentNumber >= PendingSegments.Num())
+		if (IsMalformed() || !IsInitialized() || SegmentData.Num() == 0)
 		{
-			// temp hack sanity check
 			return;
 		}
-
-		LastSegmentTime = CurrentTime;
+		else if (!PendingSegments.IsValidIndex(SegmentNumber))
+		{
+			UE_LOG(LogUdpMessaging, Warning, TEXT("FUdpReassembledMessage::Reassemble: Invalid Segment Number, ignoring segment. Message=%s, Segment=%d"),
+				*Describe(), SegmentNumber);
+			return;
+		}
 
 		if (PendingSegments[SegmentNumber])
 		{
-			if (SegmentOffset + SegmentData.Num() <= Data.Num())
+			const int32 MaxAllowableOffset = Data.Num() - SegmentData.Num();
+			if (SegmentOffset < 0 || SegmentOffset > MaxAllowableOffset)
 			{
-				FMemory::Memcpy(Data.GetData() + SegmentOffset, SegmentData.GetData(), SegmentData.Num());
-
-				PendingSegments[SegmentNumber] = false;
-				--PendingSegmentsCount;
-				ReceivedBytes += SegmentData.Num();
+				UE_LOG(LogUdpMessaging, Warning, TEXT("FUdpReassembledMessage::Reassemble: SegmentOffset is invalid or would overflow Data. Message=%s, Segment=%d, SegmentOffset=%d, SegmentSize=%d, DataSize=%d"),
+					*Describe(), SegmentNumber, SegmentOffset, SegmentData.Num(), Data.Num());
+				return;
 			}
+
+			FMemory::Memcpy(Data.GetData() + SegmentOffset, SegmentData.GetData(), SegmentData.Num());
+
+			PendingSegments[SegmentNumber] = false;
+			--PendingSegmentsCount;
+			ReceivedBytes += SegmentData.Num();
 		}
+
+		LastSegmentTime = CurrentTime;
 		PendingAcknowledgments.AddUnique(SegmentNumber);
 	}
 
+	/**
+	 * Indicates whether or not this message was found to be malformed.
+	 * Once we detect a message is malformed, we will not attempt to process any data for it.
+	 */
+	const bool IsMalformed() const
+	{
+		return bIsMalformed;
+	}
+
+	FString Describe() const
+	{
+		return FString::Printf(TEXT("Sender=%s, Sequence=%llu"), *Sender.ToString(), Sequence);
+	}
+
 private:
+
 	/** Holds the message protocol version. */
 	uint8 ProtocolVersion;
 
@@ -288,4 +335,7 @@ private:
 
 	/** Holds the message sequence. */
 	uint64 Sequence;
+
+	/** Indicates whether or not this message is malformed. */
+	bool bIsMalformed = false;
 };
