Fixed empty udp messaging command line argument overriding udp messaging settings static endpoints

diff --git a/Engine/Plugins/Messaging/UdpMessaging/Source/UdpMessaging/Private/UdpMessagingModule.cpp b/Engine/Plugins/Messaging/UdpMessaging/Source/UdpMessaging/Private/UdpMessagingModule.cpp
--- a/Engine/Plugins/Messaging/UdpMessaging/Source/UdpMessaging/Private/UdpMessagingModule.cpp
+++ b/Engine/Plugins/Messaging/UdpMessaging/Source/UdpMessaging/Private/UdpMessagingModule.cpp
@@ -537,7 +537,12 @@ protected:
 			
 			FString StaticEndpoints;
 			FParse::Value(CommandLine, TEXT("-UDPMESSAGING_TRANSPORT_STATIC="), StaticEndpoints, false);
-			StaticEndpoints.ParseIntoArray(Settings->StaticEndpoints, TEXT(","));
+			TArray<FString> CommandLineStaticEndpoints;
+			StaticEndpoints.ParseIntoArray(CommandLineStaticEndpoints, TEXT(","));
+			for (const FString& CmdStaticEndpoint : CommandLineStaticEndpoints)
+			{
+				Settings->StaticEndpoints.AddUnique(CmdStaticEndpoint);
+			}
 		}
 	}
 
