Make UDP Messaging module add/remove endpoint thread safe

diff --git a/Engine/Plugins/Messaging/UdpMessaging/Source/UdpMessaging/Private/UdpMessagingModule.cpp b/Engine/Plugins/Messaging/UdpMessaging/Source/UdpMessaging/Private/UdpMessagingModule.cpp
--- a/Engine/Plugins/Messaging/UdpMessaging/Source/UdpMessaging/Private/UdpMessagingModule.cpp
+++ b/Engine/Plugins/Messaging/UdpMessaging/Source/UdpMessaging/Private/UdpMessagingModule.cpp
@@ -359,21 +359,22 @@ public:
 
 	virtual void ShutdownServices()
 	{
-		AdditionalStaticEndpoints.Empty();
 		ShutdownBridge();
 #if PLATFORM_DESKTOP
 		ShutdownTunnel();
 #endif
+		AdditionalStaticEndpoints.Empty();
 	}
 
 	virtual void AddEndpoint(const FString& InEndpoint)
 	{
-		FIPv4Endpoint OutEndpoint;
-		if (FIPv4Endpoint::Parse(InEndpoint, OutEndpoint) && !AdditionalStaticEndpoints.Contains(OutEndpoint))
+		if (auto Transport = WeakBridgeTransport.Pin())
 		{
-			AdditionalStaticEndpoints.Add(OutEndpoint);
-			if (auto Transport = WeakBridgeTransport.Pin())
+			FScopeLock StaticEndpointsLock(&StaticEndpointsCS);
+			FIPv4Endpoint OutEndpoint;
+			if (FIPv4Endpoint::Parse(InEndpoint, OutEndpoint) && !AdditionalStaticEndpoints.Contains(OutEndpoint))
 			{
+				AdditionalStaticEndpoints.Add(OutEndpoint);
 				Transport->AddStaticEndpoint(OutEndpoint);
 			}
 		}
@@ -381,12 +382,13 @@ public:
 
 	virtual void RemoveEndpoint(const FString& InEndpoint)
 	{
-		FIPv4Endpoint OutEndpoint;
-		if (FIPv4Endpoint::Parse(InEndpoint, OutEndpoint) && AdditionalStaticEndpoints.Contains(OutEndpoint))
+		if (auto Transport = WeakBridgeTransport.Pin())
 		{
-			AdditionalStaticEndpoints.Remove(OutEndpoint);
-			if (auto Transport = WeakBridgeTransport.Pin())
+			FScopeLock StaticEndpointsLock(&StaticEndpointsCS);
+			FIPv4Endpoint OutEndpoint;
+			if (FIPv4Endpoint::Parse(InEndpoint, OutEndpoint) && AdditionalStaticEndpoints.Contains(OutEndpoint))
 			{
+				AdditionalStaticEndpoints.Remove(OutEndpoint);
 				Transport->RemoveStaticEndpoint(OutEndpoint);
 			}
 		}
@@ -665,6 +667,9 @@ private:
 	/** Holds the bridge transport if present.  */
 	TWeakPtr<FUdpMessageTransport, ESPMode::ThreadSafe> WeakBridgeTransport;
 
+	/** Critical section protecting access to the transport static endpoints and additional static endpoints. */
+	FCriticalSection StaticEndpointsCS;
+
 	/** Holds additional static endpoints added through the modular feature interface. */
 	TSet<FIPv4Endpoint> AdditionalStaticEndpoints;
 
