Copy of 11198509 from //UE4/Dev-Enterprise:

Allow TargetRules to explicitly allow UDP/TCP messaging in Shipping configuration

diff --git a/Engine/Plugins/Messaging/TcpMessaging/Source/TcpMessaging/Private/TcpMessagingModule.cpp b/Engine/Plugins/Messaging/TcpMessaging/Source/TcpMessaging/Private/TcpMessagingModule.cpp
--- a/Engine/Plugins/Messaging/TcpMessaging/Source/TcpMessaging/Private/TcpMessagingModule.cpp
+++ b/Engine/Plugins/Messaging/TcpMessaging/Source/TcpMessaging/Private/TcpMessagingModule.cpp
@@ -272,7 +272,7 @@ protected:
 	 */
 	bool SupportsNetworkedTransport() const
 	{
-#if UE_BUILD_SHIPPING
+#if UE_BUILD_SHIPPING && !(defined(ALLOW_TCP_MESSAGING_SHIPPING) && ALLOW_TCP_MESSAGING_SHIPPING)
 		return false;
 #else
 		// disallow unsupported platforms
diff --git a/Engine/Plugins/Messaging/UdpMessaging/Source/UdpMessaging/Private/UdpMessagingModule.cpp b/Engine/Plugins/Messaging/UdpMessaging/Source/UdpMessaging/Private/UdpMessagingModule.cpp
--- a/Engine/Plugins/Messaging/UdpMessaging/Source/UdpMessaging/Private/UdpMessagingModule.cpp
+++ b/Engine/Plugins/Messaging/UdpMessaging/Source/UdpMessaging/Private/UdpMessagingModule.cpp
@@ -165,7 +165,7 @@ public:
 				Ar.Logf(TEXT("    Total Bytes Out: %i"), MessageTunnel->GetTotalOutboundBytes());
 
 				TArray<TSharedPtr<IUdpMessageTunnelConnection>> Connections;
-			
+
 				if (MessageTunnel->GetConnections(Connections) > 0)
 				{
 					Ar.Log(TEXT("  Active Connections:"));
@@ -312,7 +312,7 @@ public:
 
 	virtual bool IsSupportEnabled() const
 	{
-#if !IS_PROGRAM && UE_BUILD_SHIPPING
+#if  !IS_PROGRAM && UE_BUILD_SHIPPING && !(defined(ALLOW_UDP_MESSAGING_SHIPPING) && ALLOW_UDP_MESSAGING_SHIPPING)
 		return false;
 #else
 		// disallow unsupported platforms
@@ -433,7 +433,7 @@ protected:
 
 		// Initialize the service with the additional endpoints added through the modular interface
 		TArray<FIPv4Endpoint> StaticEndpoints = AdditionalStaticEndpoints.Array();
-		
+
 		for (auto& StaticEndpoint : Settings->StaticEndpoints)
 		{
 			FIPv4Endpoint Endpoint;
@@ -451,7 +451,7 @@ protected:
 		if (Settings->MulticastTimeToLive == 0)
 		{
 			Settings->MulticastTimeToLive = 1;
-			ResaveSettings = true;		
+			ResaveSettings = true;
 		}
 
 		if (ResaveSettings)
@@ -540,7 +540,7 @@ protected:
 			FParse::Bool(CommandLine, TEXT("-UDPMESSAGING_TRANSPORT_ENABLE="), Settings->EnableTransport);
 			FParse::Value(CommandLine, TEXT("-UDPMESSAGING_TRANSPORT_UNICAST="), Settings->UnicastEndpoint);
 			FParse::Value(CommandLine, TEXT("-UDPMESSAGING_TRANSPORT_MULTICAST="), Settings->MulticastEndpoint);
-			
+
 			FString StaticEndpoints;
 			FParse::Value(CommandLine, TEXT("-UDPMESSAGING_TRANSPORT_STATIC="), StaticEndpoints, false);
 			TArray<FString> CommandLineStaticEndpoints;
@@ -573,7 +573,7 @@ protected:
 		{
 			MessageTunnel->StopServer();
 			MessageTunnel.Reset();
-		}		
+		}
 	}
 #endif
 
