Remove unused private fields - Messaging

diff --git a/Engine/Plugins/Messaging/UdpMessaging/Source/UdpMessaging/Private/Tests/UdpMessageTransportTest.cpp b/Engine/Plugins/Messaging/UdpMessaging/Source/UdpMessaging/Private/Tests/UdpMessageTransportTest.cpp
--- a/Engine/Plugins/Messaging/UdpMessaging/Source/UdpMessaging/Private/Tests/UdpMessageTransportTest.cpp
+++ b/Engine/Plugins/Messaging/UdpMessaging/Source/UdpMessaging/Private/Tests/UdpMessageTransportTest.cpp
@@ -21,7 +21,6 @@ public:
 
 	FUdpMessageTransportTestState(FAutomationTestBase& InTest, const FIPv4Endpoint& UnicastEndpoint, const FIPv4Endpoint& MulticastEndpoint, uint8 MulticastTimeToLive)
 		: NumReceivedMessages(0)
-		, Test(InTest)
 	{
 		TArray<FIPv4Endpoint> StaticEndpoints;
 		Transport = MakeShared<FUdpMessageTransport, ESPMode::ThreadSafe>(UnicastEndpoint, MulticastEndpoint, MoveTemp(StaticEndpoints), MulticastTimeToLive);
@@ -83,7 +82,6 @@ private:
 	TArray<FGuid> DiscoveredNodes;
 	TArray<FGuid> LostNodes;
 	int32 NumReceivedMessages;
-	FAutomationTestBase& Test;
 	TSharedPtr<IMessageTransport, ESPMode::ThreadSafe> Transport;
 };
 
diff --git a/Engine/Plugins/Messaging/UdpMessaging/Source/UdpMessaging/Private/Transport/UdpMessageResequencer.h b/Engine/Plugins/Messaging/UdpMessaging/Source/UdpMessaging/Private/Transport/UdpMessageResequencer.h
--- a/Engine/Plugins/Messaging/UdpMessaging/Source/UdpMessaging/Private/Transport/UdpMessageResequencer.h
+++ b/Engine/Plugins/Messaging/UdpMessaging/Source/UdpMessaging/Private/Transport/UdpMessageResequencer.h
@@ -25,7 +25,6 @@ public:
 	 */
 	FUdpMessageResequencer(uint16 InWindowSize)
 		: NextSequence(1)
-		, WindowSize(InWindowSize)
 	{ }
 
 public:
@@ -94,12 +93,6 @@ private:
 	/** Holds the next expected sequence number. */
 	uint64 NextSequence;
 
-	/** Holds the highest received sequence number. */
-	uint64 HighestReceivedSequence;
-
 	/** Holds the messages that need to be resequenced. */
 	TArray<TSharedPtr<FUdpReassembledMessage, ESPMode::ThreadSafe>> MessageHeap;
-
-	/** Holds the maximum resequence window size. */
-	uint16 WindowSize;
 };
diff --git a/Engine/Plugins/Messaging/UdpMessaging/Source/UdpMessaging/Private/Transport/UdpMessageTransport.cpp b/Engine/Plugins/Messaging/UdpMessaging/Source/UdpMessaging/Private/Transport/UdpMessageTransport.cpp
--- a/Engine/Plugins/Messaging/UdpMessaging/Source/UdpMessaging/Private/Transport/UdpMessageTransport.cpp
+++ b/Engine/Plugins/Messaging/UdpMessaging/Source/UdpMessaging/Private/Transport/UdpMessageTransport.cpp
@@ -25,7 +25,6 @@
 
 FUdpMessageTransport::FUdpMessageTransport(const FIPv4Endpoint& InUnicastEndpoint, const FIPv4Endpoint& InMulticastEndpoint, TArray<FIPv4Endpoint> InStaticEndpoints, uint8 InMulticastTtl)
 	: MessageProcessor(nullptr)
-	, MessageProcessorThread(nullptr)
 	, MulticastEndpoint(InMulticastEndpoint)
 	, MulticastReceiver(nullptr)
 	, MulticastSocket(nullptr)
diff --git a/Engine/Plugins/Messaging/UdpMessaging/Source/UdpMessaging/Private/Transport/UdpMessageTransport.h b/Engine/Plugins/Messaging/UdpMessaging/Source/UdpMessaging/Private/Transport/UdpMessageTransport.h
--- a/Engine/Plugins/Messaging/UdpMessaging/Source/UdpMessaging/Private/Transport/UdpMessageTransport.h
+++ b/Engine/Plugins/Messaging/UdpMessaging/Source/UdpMessaging/Private/Transport/UdpMessageTransport.h
@@ -115,9 +115,6 @@ private:
 	/** Holds the message processor. */
 	FUdpMessageProcessor* MessageProcessor;
 
-	/** Holds the message processor thread. */
-	FRunnableThread* MessageProcessorThread;
-
 	/** Holds the multicast endpoint. */
 	FIPv4Endpoint MulticastEndpoint;
 
