Fix timing issue between UdpMessageProcessor initialization and passing in the static endpoints to the beacon

diff --git a/Engine/Plugins/Messaging/UdpMessaging/Source/UdpMessaging/Private/Transport/UdpMessageProcessor.cpp b/Engine/Plugins/Messaging/UdpMessaging/Source/UdpMessaging/Private/Transport/UdpMessageProcessor.cpp
--- a/Engine/Plugins/Messaging/UdpMessaging/Source/UdpMessaging/Private/Transport/UdpMessageProcessor.cpp
+++ b/Engine/Plugins/Messaging/UdpMessaging/Source/UdpMessaging/Private/Transport/UdpMessageProcessor.cpp
@@ -38,8 +38,11 @@ FUdpMessageProcessor::FUdpMessageProcessor(FSocket& InSocket, const FGuid& InNod
 	, Socket(&InSocket)
 	, SocketSender(nullptr)
 	, Stopping(false)
+	, bIsInitialized(false)
 	, MessageFormat(GetDefault<UUdpMessagingSettings>()->MessageFormat) // NOTE: When the message format changes (in the Udp Messaging settings panel), the service is restarted and the processor recreated.
 {
+	Init();
+
 	WorkEvent = MakeShareable(FPlatformProcess::GetSynchEventFromPool(), [](FEvent* EventToDelete)
 	{
 		FPlatformProcess::ReturnSynchEventToPool(EventToDelete);
@@ -173,18 +176,21 @@ FSingleThreadRunnable* FUdpMessageProcessor::GetSingleThreadInterface()
 
 bool FUdpMessageProcessor::Init()
 {
-	Beacon = new FUdpMessageBeacon(Socket, LocalNodeId, MulticastEndpoint);
-	SocketSender = new FUdpSocketSender(Socket, TEXT("FUdpMessageProcessor.Sender"));
-
-	// Current protocol version 14
-	SupportedProtocolVersions.Add(UDP_MESSAGING_TRANSPORT_PROTOCOL_VERSION);
-	// Support Protocol version 10, 11, 12, 13
-	SupportedProtocolVersions.Add(13);
-	SupportedProtocolVersions.Add(12);
-	SupportedProtocolVersions.Add(11);
-	SupportedProtocolVersions.Add(10);
+	if (!bIsInitialized)
+	{
+		Beacon = new FUdpMessageBeacon(Socket, LocalNodeId, MulticastEndpoint);
+		SocketSender = new FUdpSocketSender(Socket, TEXT("FUdpMessageProcessor.Sender"));
 
-	return true;
+		// Current protocol version 14
+		SupportedProtocolVersions.Add(UDP_MESSAGING_TRANSPORT_PROTOCOL_VERSION);
+		// Support Protocol version 10, 11, 12, 13
+		SupportedProtocolVersions.Add(13);
+		SupportedProtocolVersions.Add(12);
+		SupportedProtocolVersions.Add(11);
+		SupportedProtocolVersions.Add(10);
+		bIsInitialized = true;
+	}
+	return bIsInitialized;
 }
 
 
diff --git a/Engine/Plugins/Messaging/UdpMessaging/Source/UdpMessaging/Private/Transport/UdpMessageProcessor.h b/Engine/Plugins/Messaging/UdpMessaging/Source/UdpMessaging/Private/Transport/UdpMessageProcessor.h
--- a/Engine/Plugins/Messaging/UdpMessaging/Source/UdpMessaging/Private/Transport/UdpMessageProcessor.h
+++ b/Engine/Plugins/Messaging/UdpMessaging/Source/UdpMessaging/Private/Transport/UdpMessageProcessor.h
@@ -467,6 +467,9 @@ private:
 	/** Holds a flag indicating that the thread is stopping. */
 	bool Stopping;
 
+	/** Holds a flag indicating if the processor is initialized. */
+	bool bIsInitialized;
+
 	/** Holds the thread object. */
 	FRunnableThread* Thread;
 
