Fix issue with UDP Messaging where large packets could be dropped based on the MaxBytes constraint imposed during ticking.

The root issues is that UdpMessageProcessor breaks out of the for(Segment : PendingSendSegments) when reaching MaxBytes. This had two issues when a packet or multiple packets in one frame exceeded MaxBytes -

* The unsent portion of unreliable packets would be skipped causing the whole package to be dropped.
* Reliable packets would be delayed by 100ms, plus early segments would be resent until they were acknowledged by the remote connection

This change updates MessageSegmenter so PendingSendSegments represents unsent packets and is updated by MessageProcessor as each segment is sent to the socket. This requires a separate bitarray to track unacknowledged packets so segments can be marked for resending if necessary.

MessageSegmenter::IsComplete() is now replaced by Message::Segmenter::IsSendingComplete() and MessageSegementer::AreAcknowledgementsComplete(). The latter always returns true for unreliable messages

Updated Segmenter test code to validate new code.

diff --git a/Engine/Plugins/Messaging/UdpMessaging/Source/UdpMessaging/Private/Tests/UdpMessageSegmenterTest.cpp b/Engine/Plugins/Messaging/UdpMessaging/Source/UdpMessaging/Private/Tests/UdpMessageSegmenterTest.cpp
--- a/Engine/Plugins/Messaging/UdpMessaging/Source/UdpMessaging/Private/Tests/UdpMessageSegmenterTest.cpp
+++ b/Engine/Plugins/Messaging/UdpMessaging/Source/UdpMessaging/Private/Tests/UdpMessageSegmenterTest.cpp
@@ -13,15 +13,15 @@
 IMPLEMENT_SIMPLE_AUTOMATION_TEST(FUdpMessageSegmenterTest, "System.Core.Messaging.Transports.Udp.UdpMessageSegmenter", EAutomationTestFlags::ApplicationContextMask | EAutomationTestFlags::EngineFilter)
 
 
-void RunSegmentationTest(FAutomationTestBase& Test, uint32 MessageSize, uint16 SegmentSize)
+void RunSegmentationTest(FAutomationTestBase& Test, uint32 MessageSize, uint16 SegmentSize, bool WithReliableMessages)
 {
 	check(MessageSize < 255u * SegmentSize);
 
 	uint16 NumSegments = (MessageSize + SegmentSize - 1) / SegmentSize;
-	Test.AddInfo(FString::Printf(TEXT("Segmenting message of size %i with %i segments of size %i..."), MessageSize, NumSegments, SegmentSize));
+	Test.AddInfo(FString::Printf(TEXT("Segmenting message of size %i with %i segments of size %i (Reliable=%d)..."), MessageSize, NumSegments, SegmentSize, WithReliableMessages));
 
 	// create a large message to segment
-	TSharedRef<FUdpSerializedMessage, ESPMode::ThreadSafe> Message = MakeShared<FUdpSerializedMessage, ESPMode::ThreadSafe>(EUdpMessageFormat::CborPlatformEndianness, UDP_MESSAGING_TRANSPORT_PROTOCOL_VERSION, EMessageFlags::None);
+	TSharedRef<FUdpSerializedMessage, ESPMode::ThreadSafe> Message = MakeShared<FUdpSerializedMessage, ESPMode::ThreadSafe>(EUdpMessageFormat::CborPlatformEndianness, UDP_MESSAGING_TRANSPORT_PROTOCOL_VERSION, WithReliableMessages ? EMessageFlags::Reliable : EMessageFlags::None);
 
 	for (uint8 SegmentIndex = 0; SegmentIndex < NumSegments; ++SegmentIndex)
 	{
@@ -38,6 +38,9 @@ void RunSegmentationTest(FAutomationTestBase& Test, uint32 MessageSize, uint16 S
 	FUdpMessageSegmenter Segmenter = FUdpMessageSegmenter(Message, SegmentSize);
 	Segmenter.Initialize();
 
+	Test.TestEqual(TEXT("Our calculation of segment count must match the segmenter"), Segmenter.GetSegmentCount(), NumSegments);
+
+
 	// invariants
 	{
 		Test.TestEqual(TEXT("The message size must match the actual message size"), Segmenter.GetMessageSize(), Message->TotalSize());
@@ -46,8 +49,8 @@ void RunSegmentationTest(FAutomationTestBase& Test, uint32 MessageSize, uint16 S
 
 	// pre-conditions
 	{
-		Test.TestEqual(TEXT("The initial number of pending segments must match the total number of segments in the message"), Segmenter.GetPendingSegmentsCount(), (uint32)NumSegments);
-		Test.TestFalse(TEXT("Segmentation of a non-empty message must not be complete initially"), Segmenter.IsComplete());
+		Test.TestEqual(TEXT("The initial number of pending segments must match the total number of segments in the message"), Segmenter.GetPendingSendSegmentsCount(), (uint32)NumSegments);
+		Test.TestFalse(TEXT("Segmentation of a non-empty message must not be complete initially"), Segmenter.IsSendingComplete());
 	}
 
 	// peeking at next pending segment
@@ -57,7 +60,7 @@ void RunSegmentationTest(FAutomationTestBase& Test, uint32 MessageSize, uint16 S
 
 		Segmenter.GetNextPendingSegment(OutData, OutSegmentNumber);
 
-		Test.TestEqual(TEXT("The number of pending segments must not change when peeking at a pending segment"), Segmenter.GetPendingSegmentsCount(), (uint32)NumSegments);
+		Test.TestEqual(TEXT("The number of pending segments must not change when peeking at a pending segment"), Segmenter.GetPendingSendSegmentsCount(), (uint32)NumSegments);
 	}
 
 	uint32 GeneratedSegmentCount = 0;
@@ -67,10 +70,35 @@ void RunSegmentationTest(FAutomationTestBase& Test, uint32 MessageSize, uint16 S
 		TArray<uint8> OutData;
 		uint32 OutSegmentNumber = 0;
 		uint32 NumInvalidSegments = 0;
+		uint32 NumSends = 0;
 
 		while (Segmenter.GetNextPendingSegment(OutData, OutSegmentNumber))
 		{
-			Segmenter.MarkAsAcknowledged(TArrayBuilder<uint32>().Add(OutSegmentNumber));
+			Test.TestFalse(TEXT("IsSendingComplete should return false for the segment until all segments are sent"), Segmenter.IsSendingComplete());
+
+			// mark as sent and check total
+			Segmenter.MarkAsSent(TArrayBuilder<uint32>().Add(OutSegmentNumber));
+			NumSends++;
+
+			// test how many are remaining
+			Test.TestEqual(TEXT("The number of pending send segments must be the total count minus those sent"), Segmenter.GetPendingSendSegmentsCount(), NumSegments - NumSends);
+
+			Test.TestEqual(TEXT("The number of bits in the pending segment array must match the returned count"), Segmenter.GetPendingSendSegments().CountSetBits(), Segmenter.GetPendingSendSegmentsCount());
+			
+			// Test reliable state
+			if (WithReliableMessages)
+			{
+				Test.TestEqual(TEXT("The number of acknowledged segments should be 0 after sending"), Segmenter.GetAcknowledgedSegmentsCount(), 0);
+				Test.TestEqual(TEXT("The number of bits in the pending acknowledgement array must be zero"), Segmenter.GetAcknowledgedSegments().CountSetBits(), 0);
+				Test.TestFalse(TEXT("The message should have outstanding acknowledgements"), Segmenter.AreAcknowledgementsComplete());
+			}
+			else
+			{
+				Test.TestEqual(TEXT("Unreliable messages should not have any received ack segments"), Segmenter.GetAcknowledgedSegmentsCount(), 0);
+				Test.TestEqual(TEXT("The bits in the pending acknowledgement array must be zero for unreliable messages"), Segmenter.GetAcknowledgedSegments().CountSetBits(), 0);
+				Test.TestTrue(TEXT("Unreliable messages should always be marked as ack's being complete"), Segmenter.AreAcknowledgementsComplete());
+			}
+
 			++GeneratedSegmentCount;
 
 			// verify segment size
@@ -100,31 +128,71 @@ void RunSegmentationTest(FAutomationTestBase& Test, uint32 MessageSize, uint16 S
 			}
 		}
 
+		Test.TestTrue(TEXT("IsSendingComplete should return true after all segments are sent"), Segmenter.IsSendingComplete());
+
+		if (WithReliableMessages)
+		{
+			Test.TestFalse(TEXT("AreAcknowledgementsComplete should return false after all segments for a reliable message are sent"), Segmenter.AreAcknowledgementsComplete());
+
+			// Simulate what would happen if we decided to retransmit a packet but then received our ack before the send occurred...
+
+			// mark a packet for retransmission
+			auto SegmentsToRetransmit = TArrayBuilder<uint16>().Add(OutSegmentNumber);
+			Segmenter.MarkForRetransmission(SegmentsToRetransmit);
+			Test.TestFalse(TEXT("Marking an outstanding segment for restransmission should result in sending being incomplete"), Segmenter.IsSendingComplete());
+
+			// Mark packets as acknowledged
+			uint32 NumAckd = 0;
+			for (uint32 i = 0; i < NumSends; i++)
+			{
+				Segmenter.MarkAsAcknowledged(TArrayBuilder<uint32>().Add(NumAckd++));
+				Test.TestEqual(TEXT("The number of bits in the pending acknowledgement array must match the internal count"), Segmenter.GetAcknowledgedSegments().CountSetBits(), Segmenter.GetAcknowledgedSegmentsCount());
+			}
+
+			// Send a packet that has already been acknowledged
+			auto SegmentsToSend = TArrayBuilder<uint32>().Add(OutSegmentNumber);
+			Segmenter.MarkAsSent(SegmentsToSend);
+			Test.TestTrue(TEXT("Retransmitting an outstanding segment should result in sending being complete"), Segmenter.IsSendingComplete());
+
+			// Even though we marked it for retransmission, the segmenter should report acks are complete
+			Test.TestTrue(TEXT("MarkForRetransmission should not cause the state of an acknowledged packet to change"), Segmenter.AreAcknowledgementsComplete());
+		}
+		else
+		{
+			Test.TestTrue(TEXT("AreAcknowledgementsComplete should always be true for an unreliable message"), Segmenter.AreAcknowledgementsComplete());
+		}
+		
+
 		Test.TestEqual(TEXT("The number of generated segments must match the total number of segments in the message"), GeneratedSegmentCount, (uint32)NumSegments);
 		Test.TestEqual(TEXT("The number of invalid segments must be zero"), NumInvalidSegments, 0u);
 	}
 
 	// post-conditions
 	{
-		Test.TestEqual(TEXT("The number of pending segments must be zero after segmentation is complete"), Segmenter.GetPendingSegmentsCount(), 0u);
-		Test.TestTrue(TEXT("Segmentation must be complete when there are no more pending segments"), Segmenter.IsComplete());
+		Test.TestEqual(TEXT("The number of pending segments must be zero after segmentation is complete"), Segmenter.GetPendingSendSegmentsCount(), 0u);
+		Test.TestTrue(TEXT("Segmentation must be complete when there are no more pending segments"), Segmenter.IsSendingComplete());
 	}
 }
 
 
 bool FUdpMessageSegmenterTest::RunTest(const FString& Parameters)
 {
-	// single partial segment
-	RunSegmentationTest(*this, 123, 1024);
+	bool ReliableOptions[] = {false, true};
 
-	// single full segment
-	RunSegmentationTest(*this, 1024, 1024);
+	for(bool WithReliable : ReliableOptions)
+	{
+		// single partial segment
+		RunSegmentationTest(*this, 123, 1024, WithReliable);
+
+		// single full segment
+		RunSegmentationTest(*this, 1024, 1024, WithReliable);
 
-	// multiple segments with partial last segment
-	RunSegmentationTest(*this, 100000, 1024);
+		// multiple segments with partial last segment
+		RunSegmentationTest(*this, 100000, 1024, WithReliable);
 
-	// multiple segments without partial last segment
-	RunSegmentationTest(*this, 131072, 1024);
+		// multiple segments without partial last segment
+		RunSegmentationTest(*this, 131072, 1024, WithReliable);
+	}
 
 	return true;
 }
diff --git a/Engine/Plugins/Messaging/UdpMessaging/Source/UdpMessaging/Private/Transport/UdpMessageProcessor.cpp b/Engine/Plugins/Messaging/UdpMessaging/Source/UdpMessaging/Private/Transport/UdpMessageProcessor.cpp
--- a/Engine/Plugins/Messaging/UdpMessaging/Source/UdpMessaging/Private/Transport/UdpMessageProcessor.cpp
+++ b/Engine/Plugins/Messaging/UdpMessaging/Source/UdpMessaging/Private/Transport/UdpMessageProcessor.cpp
@@ -302,6 +302,7 @@ void FUdpMessageProcessor::AcknowledgeReceipt(int32 MessageId, const FNodeInfo&
 
 	int32 OutSent;
 	Socket->SendTo(Writer.GetData(), Writer.Num(), OutSent, *NodeInfo.Endpoint.ToInternetAddr());
+	UE_LOG(LogUdpMessaging, Verbose, TEXT("Sending EUdpMessageSegments::Acknowledge for msg %d from %s"), MessageId, *NodeInfo.NodeId.ToString());
 }
 
 
@@ -405,11 +406,18 @@ void FUdpMessageProcessor::ConsumeOutboundMessages()
 			// Queue segmenters to the nodes we are dispatching to
 			if (RecipientNodeInfo)
 			{
+				UE_LOG(LogUdpMessaging, Verbose, TEXT("Passing %d byte message to be segement-sent to %s"), 
+					OutboundMessage.SerializedMessage->TotalSize(), *RecipientNodeInfo->NodeId.ToString());
+
 				RecipientNodeInfo->Segmenters.Add(
 					LastSentMessage,
 					MakeShared<FUdpMessageSegmenter>(OutboundMessage.SerializedMessage.ToSharedRef(), UDP_MESSAGING_SEGMENT_SIZE)
 				);
 			}
+			else
+			{
+				UE_LOG(LogUdpMessaging, Verbose, TEXT("No recipient NodeInfo found for %s"), *RecipientId.ToString());
+			}
 		}
 	}
 }
@@ -448,6 +456,9 @@ void FUdpMessageProcessor::ProcessAcknowledgeSegment(FInboundSegment& Segment, F
 	AcknowledgeChunk.Serialize(*Segment.Data, NodeInfo.ProtocolVersion);
 
 	NodeInfo.Segmenters.Remove(AcknowledgeChunk.MessageId);
+
+	UE_LOG(LogUdpMessaging, Verbose, TEXT("Received Acknowledge for %d from %s"), AcknowledgeChunk.MessageId , *NodeInfo.NodeId.ToString());
+
 }
 
 
@@ -456,15 +467,22 @@ void FUdpMessageProcessor::ProcessAcknowledgeSegmentsSegment(FInboundSegment& Se
 	FUdpMessageSegment::FAcknowledgeSegmentsChunk AcknowledgeChunk;
 	AcknowledgeChunk.Serialize(*Segment.Data, NodeInfo.ProtocolVersion);
 
+	UE_LOG(LogUdpMessaging, Verbose, TEXT("Received AcknowledgeSegments for %d from %s"), AcknowledgeChunk.MessageId, *NodeInfo.NodeId.ToString());
+
 	TSharedPtr<FUdpMessageSegmenter> Segmenter = NodeInfo.Segmenters.FindRef(AcknowledgeChunk.MessageId);
 	if (Segmenter.IsValid())
 	{
 		Segmenter->MarkAsAcknowledged(AcknowledgeChunk.Segments);
-		if (Segmenter->IsComplete())
+		if (Segmenter->IsSendingComplete() && Segmenter->AreAcknowledgementsComplete())
 		{
+			UE_LOG(LogUdpMessaging, Verbose, TEXT("Segmenter for %s is now complete. Removing"), *NodeInfo.NodeId.ToString());
 			NodeInfo.Segmenters.Remove(AcknowledgeChunk.MessageId);
 		}
 	}
+	else
+	{
+		UE_LOG(LogUdpMessaging, Verbose, TEXT("No such segmenter for message %d"), AcknowledgeChunk.MessageId);
+	}
 }
 
 
@@ -552,6 +570,12 @@ void FUdpMessageProcessor::ProcessDataSegment(FInboundSegment& Segment, FNodeInf
 		return;
 	}
 
+	UE_LOG(LogUdpMessaging, Verbose, TEXT("FUdpMessageProcessor::ProcessDataSegment: Reassembled %d bytes message %s for %s (%s)"),
+		ReassembledMessage->GetData().Num(), 
+		*ReassembledMessage->Describe(), 
+		*NodeInfo.NodeId.ToString(),
+		*NodeInfo.Endpoint.ToString());
+
 	AcknowledgeReceipt(DataChunk.MessageId, NodeInfo);
 	DeliverMessage(ReassembledMessage, NodeInfo);
 }
@@ -633,10 +657,16 @@ void FUdpMessageProcessor::ProcessRetransmitSegment(FInboundSegment& Segment, FN
 
 	TSharedPtr<FUdpMessageSegmenter> Segmenter = NodeInfo.Segmenters.FindRef(RetransmitChunk.MessageId);
 
+	UE_LOG(LogUdpMessaging, Verbose, TEXT("Received retransmit for %d from %s"), RetransmitChunk.MessageId, *NodeInfo.NodeId.ToString());
+
 	if (Segmenter.IsValid())
 	{
 		Segmenter->MarkForRetransmission(RetransmitChunk.Segments);
 	}
+	else
+	{
+		UE_LOG(LogUdpMessaging, Verbose, TEXT("No such segmenter for message %d"), RetransmitChunk.MessageId);
+	}
 }
 
 
@@ -665,6 +695,7 @@ void FUdpMessageProcessor::DeliverMessage(const TSharedPtr<FUdpReassembledMessag
 	// Do not deliver message while saving or garbage collecting since those deliveries will fail anyway...
 	if (GIsSavingPackage || IsGarbageCollecting())
 	{
+		UE_LOG(LogUdpMessaging, Verbose, TEXT("Skipping delivery of %s"), *ReassembledMessage->Describe());
 		return;
 	}
 
@@ -713,6 +744,7 @@ void FUdpMessageProcessor::UpdateKnownNodes()
 
 		if ((NodeId.IsValid()) && ((NodeInfo.LastSegmentReceivedTime + DeadHelloTimespan) <= CurrentTime))
 		{
+			UE_LOG(LogUdpMessaging, Verbose, TEXT("FUdpMessageProcessor::UpdateKnownNodes: Removing Node %s (%s)"), *NodeInfo.NodeId.ToString(), *NodeInfo.Endpoint.ToString());
 			NodeLostDelegate.ExecuteIfBound(NodeId);
 			It.RemoveCurrent();
 		}
@@ -795,7 +827,11 @@ int32 FUdpMessageProcessor::UpdateSegmenters(FNodeInfo& NodeInfo, uint32 MaxSend
 		{
 			FUdpMessageSegment::FDataChunk DataChunk;
 
-			for (TConstSetBitIterator<> BIt(Segmenter->GetPendingSegments()); BIt; ++BIt)
+			// Track the segments we sent as we'll update the segmenter to keep track
+			TArray<uint32> SentSegments;
+			SentSegments.Reserve(Segmenter->GetSegmentCount());
+
+			for (TConstSetBitIterator<> BIt(Segmenter->GetPendingSendSegments()); BIt; ++BIt)
 			{
 				Segmenter->GetPendingSegment(BIt.GetIndex(), DataChunk.Data);
 				DataChunk.SegmentNumber = BIt.GetIndex();
@@ -818,27 +854,57 @@ int32 FUdpMessageProcessor::UpdateSegmenters(FNodeInfo& NodeInfo, uint32 MaxSend
 
 				ByteSent += Writer->Num();
 
+				UE_LOG(LogUdpMessaging, Verbose, TEXT("FUdpMessageProcessor::UpdateSegmenters: Sending msg %d as segment %d/%d of %d bytes to %s"),
+					DataChunk.MessageId, 
+					DataChunk.SegmentNumber + 1, 
+					DataChunk.TotalSegments, 
+					Segmenter->GetMessageSize(), 
+					*NodeInfo.NodeId.ToString());
+
 				if (!SocketSender->Send(Writer, NodeInfo.Endpoint))
 				{
+					UE_LOG(LogUdpMessaging, Error, TEXT("FUdpMessageProcessor::UpdateSegmenters: Error sending message segment %d/%d of %d bytes to %s"),
+						DataChunk.SegmentNumber + 1, DataChunk.TotalSegments, Segmenter->GetMessageSize(), *NodeInfo.NodeId.ToString());
 					return -1;
  				}
 
+				// Track this segment
+				SentSegments.Push(BIt.GetIndex());
+
 				// if we reached the max send rate, break
 				if (ByteSent >= MaxSendBytes)
 				{
+					UE_LOG(LogUdpMessaging, Verbose, TEXT("FUdpMessageProcessor::UpdateSegmenters: Reached max BytesSent (%d of %d) on message segment %d/%d to %s"),
+						ByteSent, MaxSendBytes,
+						DataChunk.SegmentNumber + 1,
+						DataChunk.TotalSegments,
+						*NodeInfo.NodeId.ToString());
 					break;
 				}
 			}
 
-			// update sent time for reliable messages
-			if (EnumHasAnyFlags(Segmenter->GetMessageFlags(), EMessageFlags::Reliable))
+
+			// Mark those segments as sent. This removes them from the pending list
+			Segmenter->MarkAsSent(SentSegments);
+
+			// Did this pass complete everything?
+			if (Segmenter->IsSendingComplete())
 			{
-				Segmenter->UpdateSentTime(CurrentTime);
+				// update sent time for reliable messages
+				if (EnumHasAnyFlags(Segmenter->GetMessageFlags(), EMessageFlags::Reliable))
+				{
+					Segmenter->UpdateSentTime(CurrentTime);
+				}
+				else
+				{
+					UE_LOG(LogUdpMessaging, Verbose, TEXT("FUdpMessageProcessor::UpdateSegmenters: Finished with message segmenter for %s"), *NodeInfo.NodeId.ToString());
+					It.RemoveCurrent();
+				}
 			}
-			// if message isn't reliable no need to keep track
 			else
 			{
-				It.RemoveCurrent();
+				// Still work to do that will be picked up next tick
+				UE_LOG(LogUdpMessaging, Verbose, TEXT("FUdpMessageProcessor::UpdateSegmenters: Will continue processing segments for msg %d from %s on next tick"), DataChunk.MessageId , *NodeInfo.NodeId.ToString());
 			}
 		}
 		else if (Segmenter->IsInvalid())
@@ -875,17 +941,28 @@ bool FUdpMessageProcessor::UpdateReassemblers(FNodeInfo& NodeInfo)
 		// Send pending acknowledgments
 		if (ReassembledMessage->HasPendingAcknowledgements())
 		{
-			FUdpMessageSegment::FAcknowledgeSegmentsChunk AcknowledgeChunk(It.Key()/*MessageId*/, ReassembledMessage->GetPendingAcknowledgments()/*Segments*/);
+			TArray<uint32> PendingAcknowledgments = ReassembledMessage->GetPendingAcknowledgments();
+			const int32 AckCount = PendingAcknowledgments.Num();
+
+			FUdpMessageSegment::FAcknowledgeSegmentsChunk AcknowledgeChunk(It.Key()/*MessageId*/, MoveTemp(PendingAcknowledgments)/*Segments*/);
 			TSharedRef<FArrayWriter, ESPMode::ThreadSafe> Writer = MakeShared<FArrayWriter, ESPMode::ThreadSafe>();
 			{
 				*Writer << Header;
 				AcknowledgeChunk.Serialize(*Writer, Header.ProtocolVersion);				
 			}
 
+			UE_LOG(LogUdpMessaging, Verbose, TEXT("FUdpMessageProcessor::UpdateReassemblers: Sending EUdpMessageSegments::AcknowledgeSegments for %d segments for message %d from %s"), AckCount, It.Key(), *ReassembledMessage->Describe());
+
 			if (!SocketSender->Send(Writer, NodeInfo.Endpoint))
 			{
+				UE_LOG(LogUdpMessaging, Error, TEXT("FUdpMessageProcessor::UpdateReassemblers: error sending EUdpMessageSegments::AcknowledgeSegments from %s"), *NodeInfo.NodeId.ToString());
+
 				return false;
 			}
+			
+			UE_LOG(LogUdpMessaging, Verbose, TEXT("FUdpMessageProcessor::UpdateReassemblers: sending acknowledgement for reliable msg %d from %s"), 
+				It.Key(), 
+				*NodeInfo.NodeId.ToString());			
 		}
 
 		// Try to deliver completed message that couldn't be delivered the first time around
@@ -898,6 +975,14 @@ bool FUdpMessageProcessor::UpdateReassemblers(FNodeInfo& NodeInfo)
 		if (ReassembledMessage->GetLastSegmentTime() + StaleReassemblyInterval <= CurrentTime &&
 			(!EnumHasAnyFlags(ReassembledMessage->GetFlags(), EMessageFlags::Reliable) || ReassembledMessage->IsDelivered()))
 		{
+			if (!ReassembledMessage->IsDelivered())
+			{ 
+				const int ReceivedSegments = ReassembledMessage->GetTotalSegmentsCount() - ReassembledMessage->GetPendingSegmentsCount();
+				UE_LOG(LogUdpMessaging, Warning, TEXT("FUdpMessageProcessor::UpdateReassemblers Discarding %d/%d of stale message segements from %s"), 
+					ReceivedSegments,
+					ReassembledMessage->GetTotalSegmentsCount(),
+					*ReassembledMessage->Describe());
+			}
 			It.RemoveCurrent();
 		}
 	}
diff --git a/Engine/Plugins/Messaging/UdpMessaging/Source/UdpMessaging/Private/Transport/UdpMessageSegmenter.cpp b/Engine/Plugins/Messaging/UdpMessaging/Source/UdpMessaging/Private/Transport/UdpMessageSegmenter.cpp
--- a/Engine/Plugins/Messaging/UdpMessaging/Source/UdpMessaging/Private/Transport/UdpMessageSegmenter.cpp
+++ b/Engine/Plugins/Messaging/UdpMessaging/Source/UdpMessaging/Private/Transport/UdpMessageSegmenter.cpp
@@ -1,7 +1,7 @@
 // Copyright Epic Games, Inc. All Rights Reserved.
 
 #include "Transport/UdpMessageSegmenter.h"
-
+#include "UdpMessagingPrivate.h"
 #include "Serialization/Archive.h"
 
 #include "Transport/UdpSerializedMessage.h"
@@ -40,7 +40,7 @@ bool FUdpMessageSegmenter::GetNextPendingSegment(TArray<uint8>& OutData, uint32&
 		return false;
 	}
 
-	for (TConstSetBitIterator<> It(PendingSegments); It; ++It)
+	for (TConstSetBitIterator<> It(PendingSendSegments); It; ++It)
 	{
 		OutSegment = It.GetIndex();
 
@@ -75,7 +75,7 @@ bool FUdpMessageSegmenter::GetPendingSegment(uint32 InSegment, TArray<uint8>& Ou
 	}
 
 	// Max segment number for protocol 12 is INT32_MAX, if increased, this will need changing
-	if (InSegment < (uint32)PendingSegments.Num() && PendingSegments[InSegment])
+	if (InSegment < (uint32)PendingSendSegments.Num() && PendingSendSegments[InSegment])
 	{
 		uint64 SegmentOffset = static_cast<uint64>(InSegment) * SegmentSize;
 		uint64 ActualSegmentSize = MessageReader->TotalSize() - SegmentOffset;
@@ -108,8 +108,18 @@ void FUdpMessageSegmenter::Initialize()
 	if (SerializedMessage->GetState() == EUdpSerializedMessageState::Complete)
 	{
 		MessageReader = SerializedMessage->CreateReader();
-		PendingSegmentsCount = (MessageReader->TotalSize() + SegmentSize - 1) / SegmentSize;
-		PendingSegments.Init(true, PendingSegmentsCount);
+		PendingSendSegmentsCount = (MessageReader->TotalSize() + SegmentSize - 1) / SegmentSize;
+		PendingSendSegments.Init(true, PendingSendSegmentsCount);
+		if (EnumHasAnyFlags(GetMessageFlags(), EMessageFlags::Reliable))
+		{
+			AcknowlededSegments.Init(false, PendingSendSegmentsCount);
+		}
+		else
+		{
+			// Acks for unreliable messages are always zero
+			AcknowlededSegments.Init(false, 0);
+		}		
+		AcknowlededSegmentsCount = 0;
 	}
 }
 
@@ -137,15 +147,42 @@ EMessageFlags FUdpMessageSegmenter::GetMessageFlags() const
 	return SerializedMessage->GetFlags();
 }
 
+void FUdpMessageSegmenter::MarkAsSent(const TArray<uint32>& Segments)
+{
+	for (const auto& Segment : Segments)
+	{
+		if (Segment < (uint32)PendingSendSegments.Num() && PendingSendSegments[Segment])
+		{
+			--PendingSendSegmentsCount;
+			PendingSendSegments[Segment] = false;
+		
+			UE_LOG(LogUdpMessaging, Verbose, TEXT("Marking segment %d of %d as sent (%d outstanding)"), Segment+1, PendingSendSegments.Num(), PendingSendSegmentsCount);
+		}
+	}
+}
 
 void FUdpMessageSegmenter::MarkAsAcknowledged(const TArray<uint32>& Segments)
 {
-	for (const auto& Segment : Segments)
+	if (ensure(EnumHasAnyFlags(GetMessageFlags(), EMessageFlags::Reliable)))
 	{
-		if (Segment < (uint32)PendingSegments.Num() && PendingSegments[Segment])
+		for (const auto& Segment : Segments)
 		{
-			--PendingSegmentsCount;
-			PendingSegments[Segment] = false;
+			// Mark this segment as acknowledged. There's a chance segments could be acknowledged
+			// twice so we need to check state
+			if (Segment < (uint32)AcknowlededSegments.Num() && !AcknowlededSegments[Segment])
+			{
+				++AcknowlededSegmentsCount;
+				AcknowlededSegments[Segment] = true;
+				UE_LOG(LogUdpMessaging, Verbose, TEXT("Marked segment %d of %d as acknowledged (%d outstanding)"), Segment + 1, AcknowlededSegments.Num(), AcknowlededSegments.Num() - AcknowlededSegmentsCount);
+			}
+
+			// We may have queued a segment to be resent, if so there's now no need to resend it
+			if (Segment < (uint32)PendingSendSegments.Num() && PendingSendSegments[Segment])
+			{
+				--PendingSendSegmentsCount;
+				PendingSendSegments[Segment] = false;
+				UE_LOG(LogUdpMessaging, Log, TEXT("Received acknowledgment for segment %d that was queued or requeued for transmission. Will skip send"), Segment + 1);
+			}
 		}
 	}
 }
@@ -155,19 +192,65 @@ void FUdpMessageSegmenter::MarkForRetransmission(const TArray<uint16>& Segments)
 {
 	for (const auto& Segment : Segments)
 	{
-		if (Segment < PendingSegments.Num() && !PendingSegments[Segment])
+		UE_LOG(LogUdpMessaging, Verbose, TEXT("Marking segment %d of %d for retransmission"), Segment+1, PendingSendSegments.Num());
+
+		// mark this as pending send once again
+		if (Segment < PendingSendSegments.Num() && !PendingSendSegments[Segment])
 		{
-			++PendingSegmentsCount;
-			PendingSegments[Segment] = true;
+			++PendingSendSegmentsCount;
+			PendingSendSegments[Segment] = true;
 		}
+
+		// Note - we don't need to clear acknowledgments. If any segment is in transit and acknowledged after this 
+		// call we'll do that and if possible stop the pending send, and if not we don't need to wait for the ack
 	}
 }
 
+/**
+ * Marks the entire message for retransmission.
+ */
+void FUdpMessageSegmenter::MarkForRetransmission()
+{
+	UE_LOG(LogUdpMessaging, Verbose, TEXT("Marking all %d segments for retransmission"), PendingSendSegments.Num());
+
+	// mark all segments to be resent and clear any pending state
+	PendingSendSegments.Init(true, PendingSendSegments.Num());
+	PendingSendSegmentsCount = PendingSendSegments.Num();
+	
+	// Note - we don't need to clear acknowledgments. If any segment is in transit and acknowledged after this 
+	// call we'll do that and if possible stop the pending send, and if not we don't need to wait for the ack
+}
+
+
 const FTimespan FUdpMessageSegmenter::SendInterval = FTimespan::FromMilliseconds(100);
 
 bool FUdpMessageSegmenter::NeedSending(const FDateTime& CurrentTime)
 {
-	return LastSentTime + SendInterval <= CurrentTime;
+	// still have outstanding segments
+	if (PendingSendSegmentsCount > 0)
+	{
+		return true;
+	}
+
+	if (AreAcknowledgementsComplete() == false
+		&& LastSentTime + SendInterval <= CurrentTime)
+	{
+		// We have gone through a period of time where packets or acks may have been lost, 
+		// so resend any segments that have yet to be acknowledged
+		for (TBitArray<>::FConstIterator BIt(AcknowlededSegments); BIt; ++BIt)
+		{
+			const int32 Index = BIt.GetIndex();
+			if (BIt.GetValue() == false && !PendingSendSegments[Index])
+			{
+				PendingSendSegments[Index] = true;
+				++PendingSendSegmentsCount;
+			}
+		}
+
+		return true;
+	}
+
+	return false;
 }
 
 void FUdpMessageSegmenter::UpdateSentTime(const FDateTime& CurrentTime)
diff --git a/Engine/Plugins/Messaging/UdpMessaging/Source/UdpMessaging/Private/Transport/UdpMessageSegmenter.h b/Engine/Plugins/Messaging/UdpMessaging/Source/UdpMessaging/Private/Transport/UdpMessageSegmenter.h
--- a/Engine/Plugins/Messaging/UdpMessaging/Source/UdpMessaging/Private/Transport/UdpMessageSegmenter.h
+++ b/Engine/Plugins/Messaging/UdpMessaging/Source/UdpMessaging/Private/Transport/UdpMessageSegmenter.h
@@ -79,9 +79,9 @@ public:
 	 * Get the pending segments array.
 	 * @return the list of pending segments flags.
 	 */
-	const TBitArray<>& GetPendingSegments() const
+	const TBitArray<>& GetPendingSendSegments() const
 	{
-		return PendingSegments;
+		return PendingSendSegments;
 	}
 
 	/**
@@ -89,9 +89,28 @@ public:
 	 *
 	 * @return Number of pending segments.
 	 */
-	uint32 GetPendingSegmentsCount() const
+	uint32 GetPendingSendSegmentsCount() const
 	{
-		return PendingSegmentsCount;
+		return PendingSendSegmentsCount;
+	}
+
+	/**
+	 * Get the array of acknowledged segments. True represents acknowledgement. This array will be empty for unreliable messages
+	 * @return the list of segments.
+	 */
+	const TBitArray<>& GetAcknowledgedSegments() const
+	{
+		return AcknowlededSegments;
+	}
+
+	/**
+	 * Gets the number of segments that have been acknowledged. Always zero for unreliable messages
+	 *
+	 * @return Number of pending segments.
+	 */
+	uint32 GetAcknowledgedSegmentsCount() const
+	{
+		return AcknowlededSegmentsCount;
 	}
 
 	/**
@@ -101,7 +120,7 @@ public:
 	 */
 	uint32 GetSegmentCount() const
 	{
-		return PendingSegments.Num();
+		return PendingSendSegments.Num();
 	}
 
 	/** Initializes the segmenter. */
@@ -112,9 +131,19 @@ public:
 	 *
 	 * @return true if all segments were sent, false otherwise.
 	 */
-	bool IsComplete() const
+	bool IsSendingComplete() const
+	{
+		return (PendingSendSegmentsCount == 0);
+	}
+
+	/**
+	 * Checks whether all outstanding acknowledgments have been received. Always true for an unreliable message once its sent
+	 *
+	 * @return true if all acknowledgments are received
+	 */
+	bool AreAcknowledgementsComplete() const
 	{
-		return (PendingSegmentsCount == 0);
+		return (AcknowlededSegmentsCount == AcknowlededSegments.Num());
 	}
 
 	/**
@@ -140,6 +169,13 @@ public:
 	/** @return the message flags. */
 	EMessageFlags GetMessageFlags() const;
 
+	/**
+	* Marks the specified segments as sent
+	*
+	* @param Segments The acknowledged segments.
+	*/
+	void MarkAsSent(const TArray<uint32>& Segments);
+
 	/**
 	* Marks the specified segments as acknowledged.
 	*
@@ -148,13 +184,9 @@ public:
 	void MarkAsAcknowledged(const TArray<uint32>& Segments);
 
 	/**
-	 * Marks the entire message for retransmission.
+	 * Marks the entire message for retransmission. This does not reset acknowledgements as we only need an acknowledgment for a segement once
 	 */
-	void MarkForRetransmission()
-	{
-		PendingSegments.Init(true, PendingSegments.Num());
-	}
-
+	void MarkForRetransmission();
 	/**
 	 * Marks the specified segments for retransmission.
 	 *
@@ -194,11 +226,17 @@ private:
 	/** temp hack to support new transport API. */
 	FArchive* MessageReader;
 
-	/** Holds an array of bits that indicate which segments still need to be sent. */
-	TBitArray<> PendingSegments;
+	/** Holds an array of bits where true indicates which segments still need to be sent. */
+	TBitArray<> PendingSendSegments;
+
+	/** Holds an array of bits where true indicates a segment has been acknowledged. Will always zero-length for an unreliable segment. */
+	TBitArray<> AcknowlededSegments;
 
 	/** Holds the number of segments that haven't been sent yet. */
-	uint32 PendingSegmentsCount;
+	uint32 PendingSendSegmentsCount;
+
+	/** Holds the number of segments that haven't been acknowledged yet. Will always be zero for an unreliable message */
+	uint32 AcknowlededSegmentsCount;
 
 	/** Holds the segment size. */
 	uint16 SegmentSize;
diff --git a/Engine/Plugins/Messaging/UdpMessaging/Source/UdpMessaging/Private/Transport/UdpReassembledMessage.h b/Engine/Plugins/Messaging/UdpMessaging/Source/UdpMessaging/Private/Transport/UdpReassembledMessage.h
--- a/Engine/Plugins/Messaging/UdpMessaging/Source/UdpMessaging/Private/Transport/UdpReassembledMessage.h
+++ b/Engine/Plugins/Messaging/UdpMessaging/Source/UdpMessaging/Private/Transport/UdpReassembledMessage.h
@@ -281,7 +281,12 @@ public:
 		}
 
 		LastSegmentTime = CurrentTime;
-		PendingAcknowledgments.AddUnique(SegmentNumber);
+
+		// Send acknowledgement for unreliable messages
+		if (EnumHasAnyFlags(GetFlags(), EMessageFlags::Reliable))
+		{
+			PendingAcknowledgments.AddUnique(SegmentNumber);
+		}
 	}
 
 	/**
