Fix Import Volumetric LightMap Failed error when building lighting on Linux and Mac

BuildStaticLighting() will occasionally finish and call AppPreExit() which calls FUdpMessageProcessor::Stop() while there are still queued messages that haven't been sent.

This is a fix suggested by Jason which allows FUdpMessageProcessor::UpdateKnownNodes() to complete.

diff --git a/Engine/Plugins/Messaging/UdpMessaging/Source/UdpMessaging/Private/Transport/UdpMessageProcessor.cpp b/Engine/Plugins/Messaging/UdpMessaging/Source/UdpMessaging/Private/Transport/UdpMessageProcessor.cpp
--- a/Engine/Plugins/Messaging/UdpMessaging/Source/UdpMessaging/Private/Transport/UdpMessageProcessor.cpp
+++ b/Engine/Plugins/Messaging/UdpMessaging/Source/UdpMessaging/Private/Transport/UdpMessageProcessor.cpp
@@ -259,6 +259,11 @@ uint32 FUdpMessageProcessor::Run()
 		} while ((!InboundSegments.IsEmpty() || MoreToSend()) && !bStopping);
 	}
 
+	while (MoreToSend())
+	{
+		UpdateKnownNodes();
+	}
+
 	delete Beacon;
 	Beacon = nullptr;
 
