UdpMessaging: simplify error delegate handling for processor errors

This change addresses a crash that can potentially occur if a message
processing error occurs near the time the UdpMessaging module is
unloaded (typically at engine shutdown).

The crash was happening because the transport's TransportErrorDelegate
was being bound with BindRaw() to a function on the module, but a copy of
the delegate was being created by a lambda capture and scheduled to run
later on the main thread via the TaskGraph. Even though the original delegate
was being unbound at exit, the copy captured in the lambda remained bound
and could execute sometime after the module was unloaded and the UObject
system was no longer available. Trying to access the settings CDO at that point
would lead to a crash.

It turns out that the processor error does not actually need to propagate all the
way up to the module. The AutoRepair routine that is launched in response to an
error does not use any state from the module itself, and in fact is only intended to
keep trying as long as the transport remains valid. We can instead have the
transport manage the AutoRepair routine and cut out a level of delegate
invocation, and the transport no longer needs to expose a delegate publicly.

diff --git a/Engine/Plugins/Messaging/UdpMessaging/Source/UdpMessaging/Private/Transport/UdpMessageTransport.cpp b/Engine/Plugins/Messaging/UdpMessaging/Source/UdpMessaging/Private/Transport/UdpMessageTransport.cpp
--- a/Engine/Plugins/Messaging/UdpMessaging/Source/UdpMessaging/Private/Transport/UdpMessageTransport.cpp
+++ b/Engine/Plugins/Messaging/UdpMessaging/Source/UdpMessaging/Private/Transport/UdpMessageTransport.cpp
@@ -5,6 +5,7 @@
 
 #include "Common/UdpSocketBuilder.h"
 #include "Common/UdpSocketReceiver.h"
+#include "Containers/Ticker.h"
 #include "HAL/RunnableThread.h"
 #include "IMessageContext.h"
 #include "IMessageTransportHandler.h"
@@ -12,8 +13,10 @@
 #include "Serialization/ArrayReader.h"
 #include "SocketSubsystem.h"
 #include "Sockets.h"
+#include "UObject/UObjectBase.h"
 #include "Async/Async.h"
 
+#include "Shared/UdpMessagingSettings.h"
 #include "Transport/UdpReassembledMessage.h"
 #include "Transport/UdpDeserializedMessage.h"
 #include "Transport/UdpSerializedMessage.h"
@@ -49,9 +52,6 @@ FUdpMessageTransport::~FUdpMessageTransport()
 
 void FUdpMessageTransport::OnAppPreExit()
 {
-	// Remove any bound delegates. It's no longer relavent for us to send a transport error if we
-	// are in the shutdown phase.
-	TransportErrorDelegate.Unbind();
 	if (MessageProcessor)
 	{
 		MessageProcessor->WaitAsyncTaskCompletion();
@@ -81,13 +81,6 @@ void FUdpMessageTransport::RemoveStaticEndpoint(const FIPv4Endpoint& InEndpoint)
 	UE_LOG(LogUdpMessaging, Verbose, TEXT("Removed StaticEndpoint at %s"), *InEndpoint.ToString());
 }
 
-bool FUdpMessageTransport::RestartTransport()
-{
-	IMessageTransportHandler* Handler = TransportHandler;
-	StopTransport();
-	return StartTransport(*Handler);
-}
-
 /* IMessageTransport interface
  *****************************************************************************/
 
@@ -195,6 +188,8 @@ bool FUdpMessageTransport::StartTransport(IMessageTransportHandler& Handler)
 
 void FUdpMessageTransport::StopTransport()
 {
+	StopAutoRepairRoutine();
+
 	// shut down threads
 	delete MulticastReceiver;
 	MulticastReceiver = nullptr;
@@ -283,16 +278,88 @@ void FUdpMessageTransport::HandleProcessorNodeLost(const FGuid& LostNodeId)
 
 void FUdpMessageTransport::HandleProcessorError()
 {
-	if (!ErrorFuture.IsValid() && TransportErrorDelegate.IsBound())
-	{
-		ErrorFuture = Async(EAsyncExecution::TaskGraphMainThread, [ErrorDelegate = TransportErrorDelegate]()
+	if (!ErrorFuture.IsValid()) {
+		// Capture a weak pointer to this transport in the lambda to be executed later, and
+		// try to pin it again when the function actually runs. This guards against the transport
+		// being deleted in between the async task being scheduled and when it starts running.
+		TWeakPtr<FUdpMessageTransport, ESPMode::ThreadSafe> WeakTransportPtr = AsShared();
+		ErrorFuture = Async(EAsyncExecution::TaskGraphMainThread, [WeakTransport = WeakTransportPtr]()
 		{
-			ErrorDelegate.ExecuteIfBound();
+			// Bail out early if the UObject system is not initialized (e.g. at shutdown), since we
+			// won't be able to access the settings CDO even if the transport still exists.
+			if (!UObjectInitialized())
+			{
+					return;
+			}
+
+			if (TSharedPtr<FUdpMessageTransport, ESPMode::ThreadSafe> Transport = WeakTransport.Pin())
+			{
+				const UUdpMessagingSettings* Settings = GetDefault<UUdpMessagingSettings>();
+				if (Settings->bAutoRepair)
+				{
+					Transport->StartAutoRepairRoutine(Settings->AutoRepairAttemptLimit);
+				}
+				else
+				{
+					UE_LOG(LogUdpMessaging, Error, TEXT("UDP messaging encountered an error. Please restart the service for proper functionality"));
+				}
+			}
 		});
 	}
 }
 
 
+void FUdpMessageTransport::StartAutoRepairRoutine(uint32 MaxRetryAttempt)
+{
+	StopAutoRepairRoutine();
+
+	TWeakPtr<FUdpMessageTransport, ESPMode::ThreadSafe> WeakTransportPtr = AsShared();
+	FTimespan CheckDelay(0, 0, 1);
+	uint32 CheckNumber = 1;
+	AutoRepairHandle = FTicker::GetCoreTicker().AddTicker(FTickerDelegate::CreateLambda([WeakTransport = WeakTransportPtr, LastTime = FDateTime::UtcNow(), CheckDelay, CheckNumber, MaxRetryAttempt](float DeltaTime) mutable
+	{
+		QUICK_SCOPE_CYCLE_COUNTER(STAT_FUdpMessageTransport_AutoRepair);
+		bool bContinue = true;
+		FDateTime UtcNow = FDateTime::UtcNow();
+		if (LastTime + (CheckDelay * CheckNumber) <= UtcNow)
+		{
+			if (auto Transport = WeakTransport.Pin())
+			{
+				// if the restart fail, continue the routine if we are still under the retry attempt limit
+				bContinue = !Transport->RestartTransport();
+				bContinue = bContinue && CheckNumber <= MaxRetryAttempt;
+			}
+			// if we do not have a valid transport also stop the routine
+			else
+			{
+				bContinue = false;
+			}
+			++CheckNumber;
+			LastTime = UtcNow;
+		}
+		return bContinue;
+	}), 1.0f);
+	UE_LOG(LogUdpMessaging, Warning, TEXT("UDP messaging encountered an error. Auto repair routine started for reinitialization"));
+}
+
+
+void FUdpMessageTransport::StopAutoRepairRoutine()
+{
+	if (AutoRepairHandle.IsValid())
+	{
+		FTicker::GetCoreTicker().RemoveTicker(AutoRepairHandle);
+	}
+}
+
+
+bool FUdpMessageTransport::RestartTransport()
+{
+	IMessageTransportHandler* Handler = TransportHandler;
+	StopTransport();
+	return StartTransport(*Handler);
+}
+
+
 void FUdpMessageTransport::HandleSocketDataReceived(const TSharedPtr<FArrayReader, ESPMode::ThreadSafe>& Data, const FIPv4Endpoint& Sender)
 {
 	if (MessageProcessor != nullptr)
diff --git a/Engine/Plugins/Messaging/UdpMessaging/Source/UdpMessaging/Private/Transport/UdpMessageTransport.h b/Engine/Plugins/Messaging/UdpMessaging/Source/UdpMessaging/Private/Transport/UdpMessageTransport.h
--- a/Engine/Plugins/Messaging/UdpMessaging/Source/UdpMessaging/Private/Transport/UdpMessageTransport.h
+++ b/Engine/Plugins/Messaging/UdpMessaging/Source/UdpMessaging/Private/Transport/UdpMessageTransport.h
@@ -3,6 +3,7 @@
 #pragma once
 
 #include "CoreTypes.h"
+#include "Delegates/IDelegateInstance.h"
 #include "IMessageTransport.h"
 #include "Interfaces/IPv4/IPv4Endpoint.h"
 #include "IMessageAttachment.h"
@@ -64,24 +65,6 @@ public:
 	 */
 	void RemoveStaticEndpoint(const FIPv4Endpoint& InEndpoint);
 
-	/**
-	 * Returns a delegate that is executed when a socket error happened.
-	 *
-	 * @return The delegate.
-	 * @note this delegate is executed from the main thread by the task graph.
-	 */
-	DECLARE_DELEGATE(FOnError)
-	FOnError& OnTransportError()
-	{
-		return TransportErrorDelegate;
-	}
-
-	/**
-	 * Restart the transport using the same Transport handler
-	 * @return true if the restart was successful
-	 */
-	bool RestartTransport();
-
 	//~ IMessageTransport interface
 	virtual FName GetDebugName() const override;
 	virtual bool StartTransport(IMessageTransportHandler& Handler) override;
@@ -102,6 +85,18 @@ private:
 	/** */
 	void HandleProcessorError();
 
+	/** Launches a routine that attempts to restart the transport. */
+	void StartAutoRepairRoutine(uint32 MaxRetryAttempt);
+
+	/** Halts any currently running transport restart routine. */
+	void StopAutoRepairRoutine();
+
+	/**
+	 * Restart the transport using the same Transport handler.
+	 * @return true if the restart was successful
+	 */
+	bool RestartTransport();
+
 	/** Handles received socket data. */
 	void HandleSocketDataReceived(const TSharedPtr<FArrayReader, ESPMode::ThreadSafe>& Data, const FIPv4Endpoint& Sender);
 
@@ -109,8 +104,8 @@ private:
 	/** Holds any pending restart request. */
 	TFuture<void> ErrorFuture;
 
-	/** Holds the transport error delegate. */
-	FOnError TransportErrorDelegate;
+	/** Holds the delegate handle for the auto repair routine. */
+	FDelegateHandle AutoRepairHandle;
 
 	/** Holds the message processor. */
 	FUdpMessageProcessor* MessageProcessor;
diff --git a/Engine/Plugins/Messaging/UdpMessaging/Source/UdpMessaging/Private/UdpMessagingModule.cpp b/Engine/Plugins/Messaging/UdpMessaging/Source/UdpMessaging/Private/UdpMessagingModule.cpp
--- a/Engine/Plugins/Messaging/UdpMessaging/Source/UdpMessaging/Private/UdpMessagingModule.cpp
+++ b/Engine/Plugins/Messaging/UdpMessaging/Source/UdpMessaging/Private/UdpMessagingModule.cpp
@@ -8,7 +8,6 @@
 #include "Misc/CommandLine.h"
 #include "Misc/CoreDelegates.h"
 #include "Misc/App.h"
-#include "Containers/Ticker.h"
 #include "Interfaces/IPv4/IPv4Endpoint.h"
 #include "Stats/Stats.h"
 
@@ -469,7 +468,6 @@ protected:
 		UE_LOG(LogUdpMessaging, Log, TEXT("Initializing bridge on interface %s to multicast group %s."), *UnicastEndpoint.ToString(), *MulticastEndpoint.ToText().ToString());
 
 		TSharedRef<FUdpMessageTransport, ESPMode::ThreadSafe> Transport = MakeShared<FUdpMessageTransport, ESPMode::ThreadSafe>(UnicastEndpoint, MulticastEndpoint, MoveTemp(StaticEndpoints), Settings->MulticastTimeToLive);
-		Transport->OnTransportError().BindRaw(this, &FUdpMessagingModule::HandleTransportError);
 		WeakBridgeTransport = Transport;
 		MessageBridge = FMessageBridgeBuilder()
 			.UsingTransport(Transport);
@@ -563,7 +561,6 @@ protected:
 	/** Shuts down the message bridge. */
 	void ShutdownBridge()
 	{
-		StopAutoRepairRoutine();
 		WeakBridgeTransport.Reset();
 		if (MessageBridge.IsValid())
 		{
@@ -615,59 +612,6 @@ private:
 		}
 	}
 
-	void HandleTransportError()
-	{
-		const UUdpMessagingSettings* Settings = GetDefault<UUdpMessagingSettings>();
-		if (Settings->bAutoRepair)
-		{
-			StartAutoRepairRoutine(Settings->AutoRepairAttemptLimit);
-		}
-		else
-		{
-			UE_LOG(LogUdpMessaging, Error, TEXT("UDP messaging encountered an error. Please restart the service for proper functionality"));
-		}
-
-	}
-
-	void StartAutoRepairRoutine(uint32 MaxRetryAttempt)
-	{
-		StopAutoRepairRoutine();
-		FTimespan CheckDelay(0, 0, 1);
-		uint32 CheckNumber = 1;
-		AutoRepairHandle = FTicker::GetCoreTicker().AddTicker(FTickerDelegate::CreateLambda([WeakTransport = WeakBridgeTransport, LastTime = FDateTime::UtcNow(), CheckDelay, CheckNumber, MaxRetryAttempt](float DeltaTime) mutable
-		{
-			QUICK_SCOPE_CYCLE_COUNTER(STAT_FUdpMessagingModule_AutoRepair);
-			bool bContinue = true;
-			FDateTime UtcNow = FDateTime::UtcNow();
-			if (LastTime + (CheckDelay * CheckNumber) <= UtcNow)
-			{
-				if (auto Transport = WeakTransport.Pin())
-				{
-					// if the restart fail, continue the routine if we are still under the retry attempt limit
-					bContinue = !Transport->RestartTransport();
-					bContinue = bContinue && CheckNumber <= MaxRetryAttempt;
-				}
-				// if we do not have a valid transport also stop the routine
-				else
-				{
-					bContinue = false;
-				}
-				++CheckNumber;
-				LastTime = UtcNow;
-			}
-			return bContinue;
-		}), 1.0f);
-		UE_LOG(LogUdpMessaging, Warning, TEXT("UDP messaging encountered an error. Auto repair routine started for reinitialization"));
-	}
-
-	void StopAutoRepairRoutine()
-	{
-		if (AutoRepairHandle.IsValid())
-		{
-			FTicker::GetCoreTicker().RemoveTicker(AutoRepairHandle);
-		}
-	}
-
 private:
 	/** Name of the network messaging extension. */
 	static FName UdpMessagingName;
@@ -684,9 +628,6 @@ private:
 	/** Holds additional static endpoints added through the modular feature interface. */
 	TSet<FIPv4Endpoint> AdditionalStaticEndpoints;
 
-	/** Holds the delegate handle for the auto repair routine. */
-	FDelegateHandle AutoRepairHandle;
-
 #if PLATFORM_DESKTOP
 	/** Holds the message tunnel if present. */
 	TSharedPtr<IUdpMessageTunnel> MessageTunnel;
