Ease UX by customizing UDP Unicast Endpoint details to be a combo box with an editable text

diff --git a/Engine/Plugins/Messaging/UdpMessaging/Source/UdpMessaging/Private/Customization/SIpAddressComboBox.cpp b/Engine/Plugins/Messaging/UdpMessaging/Source/UdpMessaging/Private/Customization/SIpAddressComboBox.cpp
new file mode 100644
--- /dev/null
+++ b/Engine/Plugins/Messaging/UdpMessaging/Source/UdpMessaging/Private/Customization/SIpAddressComboBox.cpp
@@ -0,0 +1,82 @@
+﻿// Copyright Epic Games, Inc. All Rights Reserved.
+
+#include "SIpAddressComboBox.h"
+
+#if WITH_EDITOR
+
+#include "EditorStyleSet.h"
+#include "IPAddress.h"
+#include "SocketSubsystem.h" 
+#include "Widgets/Input/SEditableComboBox.h"
+#include "Widgets/Input/SEditableTextBox.h"
+
+namespace UE::UdpMessaging::Private
+{
+	static TArray<TSharedPtr<FString>> GetLocalNetworkInterfaceCardIPs()
+	{
+		TArray<TSharedPtr<FString>> LocalNetworkInterfaceCardIPs;
+	#if PLATFORM_WINDOWS
+		// Add the default route IP Address, only for windows
+		const FString DefaultRouteLocalAdapterAddress = TEXT("0.0.0.0:0");
+		LocalNetworkInterfaceCardIPs.Add(MakeShared<FString>(DefaultRouteLocalAdapterAddress));
+	#endif 
+		// Add the local host IP address
+		const FString LocalHostIpAddress = TEXT("127.0.0.1:0");
+		LocalNetworkInterfaceCardIPs.Add(MakeShared<FString>(LocalHostIpAddress));
+
+		TArray<TSharedPtr<FInternetAddr>> Addresses;
+		ISocketSubsystem::Get(PLATFORM_SOCKETSUBSYSTEM)->GetLocalAdapterAddresses(Addresses);
+		for (TSharedPtr<FInternetAddr> Address : Addresses)
+		{
+			// Add unique, so in ase the OS call returns with the local host or default route IP, we don't add it twice
+			LocalNetworkInterfaceCardIPs.AddUnique(MakeShared<FString>(Address->ToString(false) + ":0"));
+		}
+
+		return LocalNetworkInterfaceCardIPs;
+	}
+}
+
+void SIpAddressComboBox::Construct(const FArguments& InArgs)
+{
+	OnIPAddressSelectedDelegate = InArgs._OnIPAddressSelected;
+	LocalAdapterAddressSource = UE::UdpMessaging::Private::GetLocalNetworkInterfaceCardIPs();
+	ChildSlot
+	[
+		SAssignNew(LocalAdapterAddressComboBox, SComboBox<TSharedPtr<FString>>)
+		.OptionsSource(&LocalAdapterAddressSource)
+		.OnGenerateWidget(this, &SIpAddressComboBox::GenerateLocalAdapterAddressComboBoxEntry)
+		.OnSelectionChanged(this, &SIpAddressComboBox::OnIpAddressSelected)
+		.Content()
+		[
+			SAssignNew(IPAddressEditableTextBlock, SEditableText)
+			.Font(FEditorStyle::GetFontStyle(TEXT("PropertyWindow.NormalFont")))
+			.Text(FText::FromString(InArgs._InitialValue))
+			.OnTextCommitted(this, &SIpAddressComboBox::OnIPAddressTextCommmited)
+		]
+	];	
+}
+
+void SIpAddressComboBox::OnIpAddressSelected(TSharedPtr<FString> InAddress, ESelectInfo::Type InType)
+{
+	TSharedPtr<FString> SelectedIPAddress = LocalAdapterAddressComboBox->GetSelectedItem();
+	if (SelectedIPAddress.IsValid())
+	{
+		IPAddressEditableTextBlock->SetText(FText::FromString(*SelectedIPAddress));
+		OnIPAddressSelectedDelegate.ExecuteIfBound(*InAddress.Get());
+	}
+}
+
+void SIpAddressComboBox::OnIPAddressTextCommmited(const FText& Text, ETextCommit::Type)
+{
+	OnIPAddressSelectedDelegate.ExecuteIfBound(Text.ToString());
+}
+
+TSharedRef<SWidget> SIpAddressComboBox::GenerateLocalAdapterAddressComboBoxEntry(TSharedPtr<FString> InAddress)
+{
+	return
+		SNew(STextBlock)
+		.Text(FText::FromString(*InAddress))
+		.Font(FEditorStyle::GetFontStyle(TEXT("PropertyWindow.NormalFont")));
+}
+
+#endif
\ No newline at end of file
diff --git a/Engine/Plugins/Messaging/UdpMessaging/Source/UdpMessaging/Private/Customization/SIpAddressComboBox.h b/Engine/Plugins/Messaging/UdpMessaging/Source/UdpMessaging/Private/Customization/SIpAddressComboBox.h
new file mode 100644
--- /dev/null
+++ b/Engine/Plugins/Messaging/UdpMessaging/Source/UdpMessaging/Private/Customization/SIpAddressComboBox.h
@@ -0,0 +1,50 @@
+﻿// Copyright Epic Games, Inc. All Rights Reserved.
+
+#pragma once
+
+#if WITH_EDITOR
+#include "CoreMinimal.h"
+
+#include "Widgets/SCompoundWidget.h"
+#include "Widgets/Input/SComboBox.h"
+#include "Widgets/Input/SEditableText.h"
+
+class SEditableTextBox;
+
+class SIpAddressComboBox : public SCompoundWidget
+{
+public:
+
+	DECLARE_DELEGATE_OneParam(FOnChooseIP, FString)
+	
+	SLATE_BEGIN_ARGS(SIpAddressComboBox)
+	{}
+	SLATE_ARGUMENT(FString, InitialValue)
+	SLATE_EVENT(FOnChooseIP, OnIPAddressSelected)
+	SLATE_END_ARGS()
+
+	void Construct(const FArguments& InArgs);
+	
+private:
+	
+	TArray<TSharedPtr<FString>> LocalAdapterAddressSource;
+
+	TSharedPtr<SComboBox<TSharedPtr<FString>>> LocalAdapterAddressComboBox;
+	
+	/** Text box shown on top of the local adapter address address combo box */
+	TSharedPtr<SEditableText> IPAddressEditableTextBlock;
+
+	/** Delegate executed when a local IP address was selected */
+	FOnChooseIP OnIPAddressSelectedDelegate;
+
+	/** Handles changes in the local adapter address combo box */
+	void OnIpAddressSelected(TSharedPtr<FString> InAddress, ESelectInfo::Type InType);
+
+	/** Called when the IP Address was commited in the editable text block */
+	void OnIPAddressTextCommmited(const FText&, ETextCommit::Type);
+
+	/** Generates an entry in the local adapter address combo box */
+	TSharedRef<SWidget> GenerateLocalAdapterAddressComboBoxEntry(TSharedPtr<FString> InAddress);
+};
+
+#endif
diff --git a/Engine/Plugins/Messaging/UdpMessaging/Source/UdpMessaging/Private/Customization/UdpSettingsDetailsCustomization.cpp b/Engine/Plugins/Messaging/UdpMessaging/Source/UdpMessaging/Private/Customization/UdpSettingsDetailsCustomization.cpp
new file mode 100644
--- /dev/null
+++ b/Engine/Plugins/Messaging/UdpMessaging/Source/UdpMessaging/Private/Customization/UdpSettingsDetailsCustomization.cpp
@@ -0,0 +1,44 @@
+﻿// Copyright Epic Games, Inc. All Rights Reserved.
+
+#include "UdpSettingsDetailsCustomization.h"
+
+#if WITH_EDITOR
+#include "DetailLayoutBuilder.h"
+#include "DetailWidgetRow.h"
+#include "SIpAddressComboBox.h"
+#include "Shared/UdpMessagingSettings.h"
+
+TSharedRef<IDetailCustomization> FUdpSettingsDetailsCustomization::MakeInstance()
+{
+	return MakeShared<FUdpSettingsDetailsCustomization>();
+}
+
+void FUdpSettingsDetailsCustomization::CustomizeDetails(IDetailLayoutBuilder& DetailBuilder)
+{
+	UUdpMessagingSettings* Settings = GetMutableDefault<UUdpMessagingSettings>();
+	CustomizeStringAsIpDropDown(Settings->UnicastEndpoint, DetailBuilder.GetProperty(GET_MEMBER_NAME_CHECKED(UUdpMessagingSettings, UnicastEndpoint)), DetailBuilder);
+}
+
+void FUdpSettingsDetailsCustomization::CustomizeStringAsIpDropDown(const FString& InitialValue, TSharedPtr<IPropertyHandle> PropertyHandle, IDetailLayoutBuilder& DetailBuilder)
+{
+	IDetailPropertyRow* Row = DetailBuilder.EditDefaultProperty(PropertyHandle);
+
+	TSharedPtr<SWidget> NameWidget, ValueWidget;
+	Row->GetDefaultWidgets(NameWidget, ValueWidget);
+	
+	Row->CustomWidget()
+		.NameContent()
+		[
+			NameWidget->AsShared()
+		]
+		.ValueContent()
+		[
+			SNew(SIpAddressComboBox)
+			.InitialValue(InitialValue)
+			.OnIPAddressSelected_Lambda([PropertyHandle](const FString& SelectedIp)
+			{
+				PropertyHandle->SetPerObjectValue(0, SelectedIp);
+			})
+		];
+}
+#endif
diff --git a/Engine/Plugins/Messaging/UdpMessaging/Source/UdpMessaging/Private/Customization/UdpSettingsDetailsCustomization.h b/Engine/Plugins/Messaging/UdpMessaging/Source/UdpMessaging/Private/Customization/UdpSettingsDetailsCustomization.h
new file mode 100644
--- /dev/null
+++ b/Engine/Plugins/Messaging/UdpMessaging/Source/UdpMessaging/Private/Customization/UdpSettingsDetailsCustomization.h
@@ -0,0 +1,25 @@
+﻿// Copyright Epic Games, Inc. All Rights Reserved.
+
+#pragma once
+
+#if WITH_EDITOR
+#include "CoreMinimal.h"
+#include "IDetailCustomization.h"
+
+class IPropertyHandle;
+
+class FUdpSettingsDetailsCustomization : public IDetailCustomization
+{
+public:
+	
+	static TSharedRef<IDetailCustomization> MakeInstance();
+	
+	//~ Begin IDetailCustomization Interface
+	virtual void CustomizeDetails(IDetailLayoutBuilder& DetailBuilder) override;
+	//~ End IDetailCustomization Interface
+
+private:
+
+	void CustomizeStringAsIpDropDown(const FString& InitialValue, TSharedPtr<IPropertyHandle> PropertyHandle, IDetailLayoutBuilder& DetailBuilder);
+};
+#endif
diff --git a/Engine/Plugins/Messaging/UdpMessaging/Source/UdpMessaging/Private/UdpMessagingModule.cpp b/Engine/Plugins/Messaging/UdpMessaging/Source/UdpMessaging/Private/UdpMessagingModule.cpp
--- a/Engine/Plugins/Messaging/UdpMessaging/Source/UdpMessaging/Private/UdpMessagingModule.cpp
+++ b/Engine/Plugins/Messaging/UdpMessaging/Source/UdpMessaging/Private/UdpMessagingModule.cpp
@@ -14,6 +14,8 @@
 #if WITH_EDITOR
 	#include "ISettingsModule.h"
 	#include "ISettingsSection.h"
+	#include "PropertyEditorModule.h"
+	#include "Customization/UdpSettingsDetailsCustomization.h"
 #endif
 
 #include "Features/IModularFeatures.h"
@@ -238,7 +240,6 @@ public:
 #if WITH_EDITOR
 		// register settings
 		ISettingsModule* SettingsModule = FModuleManager::GetModulePtr<ISettingsModule>("Settings");
-
 		if (SettingsModule != nullptr)
 		{
 			ISettingsSectionPtr SettingsSection = SettingsModule->RegisterSettings("Project", "Plugins", "UdpMessaging",
@@ -252,6 +253,9 @@ public:
 				SettingsSection->OnModified().BindRaw(this, &FUdpMessagingModule::HandleSettingsSaved);
 			}
 		}
+
+		FPropertyEditorModule& PropertyModule = FModuleManager::GetModuleChecked<FPropertyEditorModule>("PropertyEditor");
+		PropertyModule.RegisterCustomClassLayout(UUdpMessagingSettings::StaticClass()->GetFName(), FOnGetDetailCustomizationInstance::CreateStatic(&FUdpSettingsDetailsCustomization::MakeInstance));
 #endif // WITH_EDITOR
 
 		// parse additional command line args
diff --git a/Engine/Plugins/Messaging/UdpMessaging/Source/UdpMessaging/UdpMessaging.Build.cs b/Engine/Plugins/Messaging/UdpMessaging/Source/UdpMessaging/UdpMessaging.Build.cs
--- a/Engine/Plugins/Messaging/UdpMessaging/Source/UdpMessaging/UdpMessaging.Build.cs
+++ b/Engine/Plugins/Messaging/UdpMessaging/Source/UdpMessaging/UdpMessaging.Build.cs
@@ -44,11 +44,23 @@ namespace UnrealBuildTool.Rules
 					new string[] {
 						"Settings",
 					});
-
+				
 				PrivateIncludePathModuleNames.AddRange(
 					new string[] {
 						"Settings",
 					});
+
+				PrivateDependencyModuleNames.AddRange(
+					new string[] {
+						"ApplicationCore",
+						"EditorStyle",
+						"Engine",
+						"InputCore",
+						"PropertyEditor",
+						"Slate",
+						"SlateCore",
+						"UnrealEd"
+					});
 			}
 		}
 	}
