UdpMessaging: Fixed a hang on exit with -NoThreading

#rnx

diff --git a/Engine/Plugins/Messaging/UdpMessaging/Source/UdpMessaging/Private/Transport/UdpMessageProcessor.cpp b/Engine/Plugins/Messaging/UdpMessaging/Source/UdpMessaging/Private/Transport/UdpMessageProcessor.cpp
--- a/Engine/Plugins/Messaging/UdpMessaging/Source/UdpMessaging/Private/Transport/UdpMessageProcessor.cpp
+++ b/Engine/Plugins/Messaging/UdpMessaging/Source/UdpMessaging/Private/Transport/UdpMessageProcessor.cpp
@@ -58,11 +58,18 @@ FUdpMessageProcessor::FUdpMessageProcessor(FSocket& InSocket, const FGuid& InNod
 
 FUdpMessageProcessor::~FUdpMessageProcessor()
 {
-	// shut down worker thread
-	Thread->Kill(true);
+	// shut down worker thread if it is still running
+	Thread->Kill();
+
 	delete Thread;
 	Thread = nullptr;
 
+	delete Beacon;
+	Beacon = nullptr;
+
+	delete SocketSender;
+	SocketSender = nullptr;
+
 	// remove all transport nodes
 	if (NodeLostDelegate.IsBound())
 	{
@@ -259,12 +266,6 @@ uint32 FUdpMessageProcessor::Run()
 		} while ((!InboundSegments.IsEmpty() || MoreToSend()) && !bStopping);
 	}
 
-	delete Beacon;
-	Beacon = nullptr;
-
-	delete SocketSender;
-	SocketSender = nullptr;
-
 	return 0;
 }
 
@@ -278,14 +279,11 @@ void FUdpMessageProcessor::Stop()
 
 void FUdpMessageProcessor::WaitAsyncTaskCompletion()
 {
-	// Stop the processor thread
+	// Stop to prevent any new work from being queued.
 	Stop();
 
-	// Make sure we stopped, so we can access KnownNodes safely
-	while (SocketSender != nullptr)
-	{
-		FPlatformProcess::Sleep(0); // Yield.
-	}
+	// Wait for the processor thread, so we can access KnownNodes safely
+	Thread->WaitForCompletion();
 
 	// Check if processor has in-flight serialization task(s).
 	auto HasIncompleteSerializationTasks = [this]()
