Fix ASan alloc/delete mis-match

diff --git a/Engine/Plugins/Messaging/UdpMessaging/Source/UdpMessaging/Private/Tests/UdpSerializeMessageTaskTest.cpp b/Engine/Plugins/Messaging/UdpMessaging/Source/UdpMessaging/Private/Tests/UdpSerializeMessageTaskTest.cpp
--- a/Engine/Plugins/Messaging/UdpMessaging/Source/UdpMessaging/Private/Tests/UdpSerializeMessageTaskTest.cpp
+++ b/Engine/Plugins/Messaging/UdpMessaging/Source/UdpMessaging/Private/Tests/UdpSerializeMessageTaskTest.cpp
@@ -7,6 +7,7 @@
 
 #include "UdpMessagingPrivate.h"
 #include "UdpMessagingSettings.h"
+#include "MessageEndpoint.h"
 #include "Transport/UdpSerializedMessage.h"
 #include "IMessageContext.h"
 #include "Transport/UdpSerializeMessageTask.h"
@@ -29,7 +30,7 @@ bool FUdpSerializeMessageTaskTest::RunTest(const FString& Parameters)
 	auto RunTestWithMessageFormat = [&](EUdpMessageFormat MessageFormat, uint8 ProtocolVersion)
 	{
 		const auto TimeSent = FDateTime(2015, 9, 17, 10, 59, 23, 666);
-		const auto Context = MakeShared<FUdpMockMessageContext, ESPMode::ThreadSafe>(new FUdpMockMessage, TimeSent);
+		const auto Context = MakeShared<FUdpMockMessageContext, ESPMode::ThreadSafe>(FMessageEndpoint::MakeMessage<FUdpMockMessage>(), TimeSent);
 
 		// synchronous reference serialization
 		const auto Message1 = MakeShared<FUdpSerializedMessage, ESPMode::ThreadSafe>(MessageFormat, ProtocolVersion, Context->GetFlags());
