Extend INetworkMessagingExtension to support network monitoring.

diff --git a/Engine/Plugins/Messaging/UdpMessaging/Source/UdpMessaging/Private/Transport/UdpMessageProcessor.cpp b/Engine/Plugins/Messaging/UdpMessaging/Source/UdpMessaging/Private/Transport/UdpMessageProcessor.cpp
--- a/Engine/Plugins/Messaging/UdpMessaging/Source/UdpMessaging/Private/Transport/UdpMessageProcessor.cpp
+++ b/Engine/Plugins/Messaging/UdpMessaging/Source/UdpMessaging/Private/Transport/UdpMessageProcessor.cpp
@@ -3,6 +3,7 @@
 #include "Transport/UdpMessageProcessor.h"
 #include "Algo/AllOf.h"
 #include "Algo/Transform.h"
+#include "INetworkMessagingExtension.h"
 #include "UdpMessagingPrivate.h"
 
 #include "Common/UdpSocketSender.h"
@@ -30,6 +31,17 @@
 const int32 FUdpMessageProcessor::DeadHelloIntervals = 5;
 
 
+namespace UE::Private::MessageProcessor
+{
+
+FOnTransferDataUpdated& OnSegmenterUpdated()
+{
+	static FOnTransferDataUpdated OnTransferUpdated;
+
+	return OnTransferUpdated;
+}
+
+}
 /* FUdpMessageProcessor structors
  *****************************************************************************/
 
@@ -100,6 +112,16 @@ void FUdpMessageProcessor::RemoveStaticEndpoint(const FIPv4Endpoint& InEndpoint)
 	}
 }
 
+TArray<FIPv4Endpoint> FUdpMessageProcessor::GetKnownEndpoints() const
+{
+	TArray<FIPv4Endpoint> Endpoints;
+	for (const auto& NodePair : KnownNodes)
+	{
+		Endpoints.Add(NodePair.Value.Endpoint);
+	}
+	return Endpoints;
+}
+
 /* FUdpMessageProcessor interface
  *****************************************************************************/
 
@@ -176,10 +198,10 @@ bool FUdpMessageProcessor::EnqueueOutboundMessage(const TSharedRef<IMessageConte
 	return true;
 }
 
-FUdpMessageTransportStatistics FUdpMessageProcessor::GetStats(FGuid Node) const
+FMessageTransportStatistics FUdpMessageProcessor::GetStats(FGuid Node) const
 {
 	FScopeLock NodeVersionLock(&StatisticsCS);
-	if (FUdpMessageTransportStatistics const* Stats = NodeStats.Find(Node))
+	if (FMessageTransportStatistics const* Stats = NodeStats.Find(Node))
 	{
 		return *Stats;
 	}
@@ -188,6 +210,7 @@ FUdpMessageTransportStatistics FUdpMessageProcessor::GetStats(FGuid Node) const
 
 void FUdpMessageProcessor::SendSegmenterStatsToListeners(int32 MessageId, FGuid NodeId, const TSharedPtr<FUdpMessageSegmenter>& Segmenter)
 {
+	FOnTransferDataUpdated& SegmenterUpdatedDelegate = UE::Private::MessageProcessor::OnSegmenterUpdated();
 	if(!SegmenterUpdatedDelegate.IsBound())
 	{
 		return;
@@ -198,13 +221,14 @@ void FUdpMessageProcessor::SendSegmenterStatsToListeners(int32 MessageId, FGuid
 		return;
 	}
 	uint32 SegmentCount = Segmenter->GetSegmentCount();
-	SegmenterUpdatedDelegate.Execute(
+	SegmenterUpdatedDelegate.Broadcast(
 		{
 			NodeId,
 			MessageId,
-			SegmentCount - Segmenter->GetPendingSendSegmentsCount(),
-			Segmenter->GetAcknowledgedSegmentsCount(),
-			SegmentCount
+			// Convert segment data into bytes.
+			UDP_MESSAGING_SEGMENT_SIZE * SegmentCount,
+			UDP_MESSAGING_SEGMENT_SIZE * (SegmentCount - Segmenter->GetPendingSendSegmentsCount()),
+			UDP_MESSAGING_SEGMENT_SIZE * (Segmenter->GetAcknowledgedSegmentsCount())
 		});
 }
 
@@ -642,7 +666,7 @@ void FUdpMessageProcessor::ProcessDataSegment(FInboundSegment& Segment, FNodeInf
 		}
 	}
 
-	NodeInfo.Statistics.SegmentsReceived++;
+	NodeInfo.Statistics.PacketsReceived++;
 	ReassembledMessage->Reassemble(DataChunk.SegmentNumber, DataChunk.SegmentOffset, DataChunk.Data, CurrentTime);
 
 	// Deliver or re-sequence message
@@ -961,7 +985,7 @@ FSentSegmentInfo FUdpMessageProcessor::SendNextSegmentForMessageId(FNodeInfo& No
 
 		Segmenter->MarkAsSent(DataChunk.SegmentNumber);
 
-		NodeInfo.Statistics.SegmentsSent++;
+		NodeInfo.Statistics.PacketsSent++;
 		SentInfo.BytesSent += Writer->Num();
 		SentInfo.SequenceNumber = ++NodeInfo.SequenceId;
 		SentInfo.bIsReliable = EnumHasAnyFlags(Segmenter->GetMessageFlags(), EMessageFlags::Reliable);
@@ -987,9 +1011,9 @@ FSentSegmentInfo FUdpMessageProcessor::SendNextSegmentForMessageId(FNodeInfo& No
 		SentInfo.bRequiresRequeue = true;
 	}
 
-	NodeInfo.Statistics.BytesSent += SentInfo.BytesSent;
-	NodeInfo.Statistics.IpAddress = NodeInfo.Endpoint;
-	NodeInfo.Statistics.SegmentsInFlight = NodeInfo.InflightSegments.Num();
+	NodeInfo.Statistics.TotalBytesSent += SentInfo.BytesSent;
+	NodeInfo.Statistics.IPv4AsString = NodeInfo.Endpoint.ToString();
+	NodeInfo.Statistics.PacketsInFlight = NodeInfo.InflightSegments.Num();
 	return MoveTemp(SentInfo);
 }
 
diff --git a/Engine/Plugins/Messaging/UdpMessaging/Source/UdpMessaging/Private/Transport/UdpMessageProcessor.h b/Engine/Plugins/Messaging/UdpMessaging/Source/UdpMessaging/Private/Transport/UdpMessageProcessor.h
--- a/Engine/Plugins/Messaging/UdpMessaging/Source/UdpMessaging/Private/Transport/UdpMessageProcessor.h
+++ b/Engine/Plugins/Messaging/UdpMessaging/Source/UdpMessaging/Private/Transport/UdpMessageProcessor.h
@@ -13,6 +13,7 @@
 #include "IMessageContext.h"
 #include "IMessageTransport.h"
 
+#include "INetworkMessagingExtension.h"
 #include "Interfaces/IPv4/IPv4Endpoint.h"
 #include "Misc/DateTime.h"
 #include "Misc/Guid.h"
@@ -107,7 +108,13 @@ struct FSentSegmentInfo
 	}
 };
 
-DECLARE_DELEGATE_OneParam(FOnSegmenterUpdated, FUdpSegmenterStats)
+namespace UE::Private::MessageProcessor
+{
+/**
+ * Global delegate for handling segmenter updates.
+ */
+FOnTransferDataUpdated& OnSegmenterUpdated();
+}
 
 /**
  * Implements a message processor for UDP messages.
@@ -153,7 +160,7 @@ class FUdpMessageProcessor
 		uint64 WindowSize = 64;
 
 		/** Various transport statistics for this endpoint */
-		FUdpMessageTransportStatistics Statistics;
+		FMessageTransportStatistics Statistics;
 
 		/** Default constructor. */
 		FNodeInfo()
@@ -259,7 +266,7 @@ class FUdpMessageProcessor
 				// In the case of segment loss half our window size to a minimum value.
 				WindowSize = FGenericPlatformMath::Max<uint64>(WindowSize/2, MinimumWindowSize);
 			}
-			Statistics.AcksReceived += NumAcks;
+			Statistics.PacketsAcked += NumAcks;
 			Statistics.WindowSize = WindowSize;
 		}
 
@@ -349,7 +356,7 @@ class FUdpMessageProcessor
 				}
 				return bIsLost;
 			});
-			Statistics.SegmentsLost += SegmentsLost;
+			Statistics.PacketsLost += SegmentsLost;
 			return SegmentsLost;
 		}
 
@@ -523,7 +530,7 @@ public:
 	/**
 	 * Get the current running network statistics for the given node.
 	 */
-	FUdpMessageTransportStatistics GetStats(FGuid Node) const;
+	FMessageTransportStatistics GetStats(FGuid Node) const;
 
 public:
 
@@ -577,13 +584,7 @@ public:
 		return ErrorDelegate;
 	}
 
-	/**
-	 * Delegate is invoked whenever a segmenter is added/removed or updated.
-	 */
-	FOnSegmenterUpdated& OnSegmenterUpdated()
-	{
-		return SegmenterUpdatedDelegate;
-	}
+	TArray<FIPv4Endpoint> GetKnownEndpoints() const;
 
 public:
 
@@ -792,7 +793,7 @@ private:
 	mutable FCriticalSection StatisticsCS;
 
 	/** Map that holds latest statistics for network transmission */
-	TMap<FGuid, FUdpMessageTransportStatistics> NodeStats;
+	TMap<FGuid, FMessageTransportStatistics> NodeStats;
 
 	/** Mutex protecting access to the NodeVersions map. */
 	mutable FCriticalSection NodeVersionCS;
@@ -842,9 +843,6 @@ private:
 	/** Holds a delegate to be invoked when a socket error happen. */
 	FOnError ErrorDelegate;
 
-	/** Holds the segmenter updated delegate */
-	FOnSegmenterUpdated SegmenterUpdatedDelegate;
-
 	/** The configured message format (from UUdpMessagingSettings). */
 	EUdpMessageFormat MessageFormat;
 
diff --git a/Engine/Plugins/Messaging/UdpMessaging/Source/UdpMessaging/Private/Transport/UdpMessageTransport.cpp b/Engine/Plugins/Messaging/UdpMessaging/Source/UdpMessaging/Private/Transport/UdpMessageTransport.cpp
--- a/Engine/Plugins/Messaging/UdpMessaging/Source/UdpMessaging/Private/Transport/UdpMessageTransport.cpp
+++ b/Engine/Plugins/Messaging/UdpMessaging/Source/UdpMessaging/Private/Transport/UdpMessageTransport.cpp
@@ -1,6 +1,7 @@
 // Copyright Epic Games, Inc. All Rights Reserved.
 
 #include "Transport/UdpMessageTransport.h"
+#include "INetworkMessagingExtension.h"
 #include "UdpMessagingPrivate.h"
 
 #include "Common/UdpSocketBuilder.h"
@@ -89,6 +90,37 @@ FName FUdpMessageTransport::GetDebugName() const
 	return "UdpMessageTransport";
 }
 
+TArray<FIPv4Endpoint> FUdpMessageTransport::GetKnownEndpoints() const
+{
+	if (MessageProcessor)
+	{
+		return MessageProcessor->GetKnownEndpoints();
+	}
+	return {};
+}
+
+FMessageTransportStatistics FUdpMessageTransport::GetLatestStatistics(FGuid NodeId) const
+{
+	if (MessageProcessor)
+	{
+		return MessageProcessor->GetStats(NodeId);
+	}
+	return {};
+}
+
+TArray<FIPv4Endpoint> FUdpMessageTransport::GetListeningAddresses() const
+{
+	TArray<FIPv4Endpoint> Listeners;
+#if PLATFORM_DESKTOP
+	Listeners.Add(UnicastEndpoint);
+#endif
+
+#if PLATFORM_SUPPORTS_UDP_MULTICAST_GROUP
+	Listeners.Add(MulticastEndpoint);
+#endif
+
+	return Listeners;
+}
 
 bool FUdpMessageTransport::StartTransport(IMessageTransportHandler& Handler)
 {
diff --git a/Engine/Plugins/Messaging/UdpMessaging/Source/UdpMessaging/Private/Transport/UdpMessageTransport.h b/Engine/Plugins/Messaging/UdpMessaging/Source/UdpMessaging/Private/Transport/UdpMessageTransport.h
--- a/Engine/Plugins/Messaging/UdpMessaging/Source/UdpMessaging/Private/Transport/UdpMessageTransport.h
+++ b/Engine/Plugins/Messaging/UdpMessaging/Source/UdpMessaging/Private/Transport/UdpMessageTransport.h
@@ -5,6 +5,7 @@
 #include "CoreTypes.h"
 #include "Delegates/IDelegateInstance.h"
 #include "IMessageTransport.h"
+#include "INetworkMessagingExtension.h"
 #include "Interfaces/IPv4/IPv4Endpoint.h"
 #include "IMessageAttachment.h"
 #include "Templates/SharedPointer.h"
@@ -72,6 +73,22 @@ public:
 	virtual void StopTransport() override;
 	virtual bool TransportMessage(const TSharedRef<IMessageContext, ESPMode::ThreadSafe>& Context, const TArray<FGuid>& Recipients) override;
 
+	/**
+	 * Get the IPv4Endpoint listening addresses.
+	 */
+	TArray<FIPv4Endpoint> GetListeningAddresses() const;
+
+	/**
+	 * Get the known endpoints that are in active communication with the UDP transport. This is information
+	 * known by the processor.
+	 */
+	TArray<FIPv4Endpoint> GetKnownEndpoints() const;
+
+	/**
+	 * Retrieve the latest network statistics for the given node.
+	 */
+	FMessageTransportStatistics GetLatestStatistics(FGuid NodeId) const;
+
 private:
 
 	/** Handles received transport messages. */
diff --git a/Engine/Plugins/Messaging/UdpMessaging/Source/UdpMessaging/Private/UdpMessagingModule.cpp b/Engine/Plugins/Messaging/UdpMessaging/Source/UdpMessaging/Private/UdpMessagingModule.cpp
--- a/Engine/Plugins/Messaging/UdpMessaging/Source/UdpMessaging/Private/UdpMessagingModule.cpp
+++ b/Engine/Plugins/Messaging/UdpMessaging/Source/UdpMessaging/Private/UdpMessagingModule.cpp
@@ -1,5 +1,6 @@
 // Copyright Epic Games, Inc. All Rights Reserved.
 
+#include "UdpMessageProcessor.h"
 #include "UdpMessagingPrivate.h"
 
 #include "CoreTypes.h"
@@ -377,7 +378,56 @@ public:
 		AdditionalStaticEndpoints.Empty();
 	}
 
-	virtual void AddEndpoint(const FString& InEndpoint)
+	virtual bool CanProvideNetworkStatistics() const override
+	{
+		return true;
+	}
+
+	virtual FMessageTransportStatistics GetLatestNetworkStatistics(FGuid NodeId) const override
+	{
+		if (auto Transport = WeakBridgeTransport.Pin())
+		{
+			return Transport->GetLatestStatistics(NodeId);
+		}
+		return {};
+	}
+
+	virtual FOnTransferDataUpdated& OnTransferUpdatedFromThread() override
+	{
+		return UE::Private::MessageProcessor::OnSegmenterUpdated();
+	}
+
+	virtual TArray<FString> GetListeningAddresses() const override
+	{
+		if (auto Transport = WeakBridgeTransport.Pin())
+		{
+			TArray<FString> StringifiedEndpoints;
+			TArray<FIPv4Endpoint> Endpoints = Transport->GetListeningAddresses();
+			for (const FIPv4Endpoint& Endpoint : Endpoints)
+			{
+				StringifiedEndpoints.Add(Endpoint.ToString());
+			}
+			return StringifiedEndpoints;
+		}
+		return {};
+	}
+
+	virtual TArray<FString> GetKnownEndpoints() const override
+	{
+		if (auto Transport = WeakBridgeTransport.Pin())
+		{
+			TArray<FIPv4Endpoint> Endpoints = Transport->GetKnownEndpoints();
+			TArray<FString> StringifiedEndpoints;
+			for (const FIPv4Endpoint& Endpoint : Endpoints)
+			{
+				StringifiedEndpoints.Add(Endpoint.ToString());
+			}
+			return StringifiedEndpoints;
+		}
+		return {};
+	}
+
+	virtual void AddEndpoint(const FString& InEndpoint) override
 	{
 		if (auto Transport = WeakBridgeTransport.Pin())
 		{
@@ -391,7 +441,7 @@ public:
 		}
 	}
 
-	virtual void RemoveEndpoint(const FString& InEndpoint)
+	virtual void RemoveEndpoint(const FString& InEndpoint) override
 	{
 		if (auto Transport = WeakBridgeTransport.Pin())
 		{
diff --git a/Engine/Source/Runtime/Messaging/Public/INetworkMessagingExtension.h b/Engine/Source/Runtime/Messaging/Public/INetworkMessagingExtension.h
--- a/Engine/Source/Runtime/Messaging/Public/INetworkMessagingExtension.h
+++ b/Engine/Source/Runtime/Messaging/Public/INetworkMessagingExtension.h
@@ -4,6 +4,61 @@
 
 #include "CoreTypes.h"
 #include "Features/IModularFeature.h"
+#include "Delegates/Delegate.h"
+#include "Misc/Guid.h"
+#include "Misc/Timespan.h"
+
+/** General information on transport health per Node Id. */
+struct FMessageTransportStatistics
+{
+	/** Total number of bytes sent to the destination endpoint.*/
+	uint64 TotalBytesSent = 0;
+
+	/** Total number of packets sent to the destination. */
+	uint64 PacketsSent = 0;
+
+	/** Total number of packets lost to the destination. */
+	uint64 PacketsLost = 0;
+
+	/** Number of packets acknowledged. */
+	uint64 PacketsAcked = 0;
+
+	/** Number of packets received by this endpoint. */
+	uint64 PacketsReceived = 0;
+
+	/** Current packets in flight waiting for */
+	uint64 PacketsInFlight = 0;
+
+	/** The size of our sending window (as indicated by the transport).  */
+	uint32 WindowSize = 0;
+
+	/** Computed average Round-trip time to receive data from connected endpoint. */
+	FTimespan AverageRTT{0};
+
+	/** IPv4 address as a string value. */
+	FString IPv4AsString;
+};
+
+/** Per-node per-message transfer statistics. */
+struct FTransferStatistics
+{
+	/** Unique Id for the target. */
+	FGuid  DestinationId;
+
+	/** Monotonically increasing ID for each message sent across the network. */
+	int32  MessageId;
+
+	/** Total bytes to send for the given MessageId */
+	uint64 BytesToSend;
+
+	/** Total bytes sent thus far. */
+	uint64 BytesSent;
+
+	/** Total number of bytes acknowledge by the destination. */
+	uint64 BytesAcknowledged;
+};
+
+DECLARE_MULTICAST_DELEGATE_OneParam(FOnTransferDataUpdated, FTransferStatistics);
 
 /**
  * Interface for the messaging module network extension
@@ -34,14 +89,42 @@ public:
 	 */
 	virtual void RestartServices() = 0;
 
+	/**
+	 * The list of network addresses that we are currently listening on.  The string must be
+	 * in the form of <address:port>
+	 */
+	virtual TArray<FString> GetListeningAddresses() const = 0;
+
+	/**
+	 * Indicates if this network messaging interface can return network statistics.
+	 */
+	virtual bool CanProvideNetworkStatistics() const = 0;
+
+	/**
+	 * Return the current network counters for the given Node endpoint
+	 */
+	virtual FMessageTransportStatistics GetLatestNetworkStatistics(FGuid NodeId) const = 0;
+
+	/**
+	 * Delegate invoked when any transmission statistics are updated. This delegate may get called from another thread.
+	 * Please consider thread safety when receiving this delegate.
+	 */
+	virtual FOnTransferDataUpdated& OnTransferUpdatedFromThread() = 0;
+
 	/**
 	 * Shutdown this messaging extension services for MessageBus
-	 * and remove any configuration modification. 
+	 * and remove any configuration modification.
 	 * Using RestartServices after ShutdownServices will start the service with an unaltered configuration
 	 * @see RestartServices
 	 */
 	virtual void ShutdownServices() = 0;
 
+	/**
+	 * Returns the list of internet addresses known by the transport. The addresses are in the form
+	 * of <address:port>
+	 */
+	virtual TArray<FString> GetKnownEndpoints() const = 0;
+
 	/**
 	 * Add an endpoint to the running configuration of this messaging service
 	 * This change is transient and does not modified saved configuration.
