Join all interfaces to the multi-cast group when using INADDR_ANY for unicast.

diff --git a/Engine/Plugins/Messaging/UdpMessaging/Source/UdpMessaging/Private/Transport/UdpMessageTransport.cpp b/Engine/Plugins/Messaging/UdpMessaging/Source/UdpMessaging/Private/Transport/UdpMessageTransport.cpp
--- a/Engine/Plugins/Messaging/UdpMessaging/Source/UdpMessaging/Private/Transport/UdpMessageTransport.cpp
+++ b/Engine/Plugins/Messaging/UdpMessaging/Source/UdpMessaging/Private/Transport/UdpMessageTransport.cpp
@@ -2,6 +2,7 @@
 
 #include "Transport/UdpMessageTransport.h"
 #include "INetworkMessagingExtension.h"
+#include "Interfaces/IPv4/IPv4Address.h"
 #include "UdpMessagingPrivate.h"
 
 #include "Common/UdpSocketBuilder.h"
@@ -122,6 +123,52 @@ TArray<FIPv4Endpoint> FUdpMessageTransport::GetListeningAddresses() const
 	return Listeners;
 }
 
+namespace UE::UdpMessaging::Private
+{
+
+void DoJoinMulticastGroup(const TSharedRef<FInternetAddr>& MulticastAddr, const TSharedPtr<FInternetAddr>& IpAddr, FSocket* MulticastSocket)
+{
+	if (!IpAddr.IsValid())
+	{
+		return;
+	}
+	const bool bJoinedGroup = MulticastSocket->JoinMulticastGroup(*MulticastAddr, *IpAddr);
+	if (bJoinedGroup)
+	{
+		UE_LOG(LogUdpMessaging, Display, TEXT("Added local interface '%s' to multicast group '%s'"),
+			   *IpAddr->ToString(false), *MulticastAddr->ToString(true));
+	}
+	else
+	{
+		UE_LOG(LogUdpMessaging, Warning, TEXT("Failed to join multicast group '%s' on detected local interface '%s'"),
+			   *MulticastAddr->ToString(true), *IpAddr->ToString(false));
+	}
+}
+
+
+void JoinedToGroup(const FIPv4Endpoint& UnicastEndpoint, const FIPv4Endpoint& MulticastEndpoint, FSocket* MulticastSocket)
+{
+#if PLATFORM_SUPPORTS_UDP_MULTICAST_GROUP
+	TSharedRef<FInternetAddr> MulticastAddr = MulticastEndpoint.ToInternetAddr();
+	if (UnicastEndpoint.Address == FIPv4Address::Any)
+	{
+		TArray<TSharedPtr<FInternetAddr>> LocapIps;
+		ISocketSubsystem::Get(PLATFORM_SOCKETSUBSYSTEM)->GetLocalAdapterAddresses(LocapIps);
+		for (const TSharedPtr<FInternetAddr>& LocalIp : LocapIps)
+		{
+			DoJoinMulticastGroup(MulticastAddr, LocalIp, MulticastSocket);
+		}
+	}
+	else
+	{
+		TSharedRef<FInternetAddr> UnicastAddr = UnicastEndpoint.ToInternetAddr();
+		DoJoinMulticastGroup(MulticastAddr, UnicastAddr , MulticastSocket);
+	}
+#endif
+}
+
+} // namespace UE::UdpMessaging::Private
+
 bool FUdpMessageTransport::StartTransport(IMessageTransportHandler& Handler)
 {
 	// Set the handler even if initialization fails. This allows tries for reinitialization using the same handler.
@@ -159,13 +206,14 @@ bool FUdpMessageTransport::StartTransport(IMessageTransportHandler& Handler)
 #endif
 		.BoundToPort(MulticastEndpoint.Port)
 #if PLATFORM_SUPPORTS_UDP_MULTICAST_GROUP
-		.JoinedToGroup(MulticastEndpoint.Address, UnicastEndpoint.Address)
 		.WithMulticastLoopback()
 		.WithMulticastTtl(MulticastTtl)
 		.WithMulticastInterface(UnicastEndpoint.Address)
 #endif
 		.WithReceiveBufferSize(UDP_MESSAGING_RECEIVE_BUFFER_SIZE);
 
+	UE::UdpMessaging::Private::JoinedToGroup(UnicastEndpoint, MulticastEndpoint, MulticastSocket);
+	
 	if (MulticastSocket == nullptr)
 	{
 		UE_LOG(LogUdpMessaging, Warning, TEXT("StartTransport failed to create multicast socket on %s, joined to %s with TTL %i"), *UnicastEndpoint.ToString(), *MulticastEndpoint.ToString(), MulticastTtl);
