Deprecating ANY_PACKAGE.

diff --git a/Engine/Plugins/FX/Niagara/Source/NiagaraEditor/Private/NiagaraEditorModule.cpp b/Engine/Plugins/FX/Niagara/Source/NiagaraEditor/Private/NiagaraEditorModule.cpp
--- a/Engine/Plugins/FX/Niagara/Source/NiagaraEditor/Private/NiagaraEditorModule.cpp
+++ b/Engine/Plugins/FX/Niagara/Source/NiagaraEditor/Private/NiagaraEditorModule.cpp
@@ -322,7 +322,7 @@ void DumpParameterStore(const FNiagaraParameterStore& ParameterStore)
 
 void DumpRapidIterationParametersForScript(UNiagaraScript* Script, const FString& HeaderName)
 {
-	UEnum* NiagaraScriptUsageEnum = FindObjectChecked<UEnum>(ANY_PACKAGE, TEXT("ENiagaraScriptUsage"), true);
+	UEnum* NiagaraScriptUsageEnum = FindObjectChecked<UEnum>(nullptr, TEXT("/Script/Niagara.ENiagaraScriptUsage"), true);
 	FString UsageName = NiagaraScriptUsageEnum->GetNameByValue((int64)Script->GetUsage()).ToString();
 	UE_LOG(LogNiagaraEditor, Log, TEXT("%s - %s - %s"), *Script->GetPathName(), *HeaderName, *UsageName);
 	DumpParameterStore(Script->RapidIterationParameters);
@@ -454,7 +454,7 @@ void PreventAllSystemRecompiles()
 	FAssetRegistryModule& AssetRegistryModule = FModuleManager::LoadModuleChecked<FAssetRegistryModule>("AssetRegistry");
 
 	TArray<FAssetData> SystemAssets;
-	AssetRegistryModule.Get().GetAssetsByClass(UNiagaraSystem::StaticClass()->GetFName(), SystemAssets);
+	AssetRegistryModule.Get().GetAssetsByClass(UNiagaraSystem::StaticClass()->GetClassPathName(), SystemAssets);
 
 	TSet<FVersionedNiagaraEmitter> CompiledEmitters;
 	int32 ItemIndex = 0;
@@ -481,7 +481,7 @@ void UpgradeAllNiagaraAssets()
 	FAssetRegistryModule& AssetRegistryModule = FModuleManager::LoadModuleChecked<FAssetRegistryModule>("AssetRegistry");
 
 	TArray<FAssetData> SystemAssets;
-	AssetRegistryModule.Get().GetAssetsByClass(UNiagaraSystem::StaticClass()->GetFName(), SystemAssets);
+	AssetRegistryModule.Get().GetAssetsByClass(UNiagaraSystem::StaticClass()->GetClassPathName(), SystemAssets);
 
 	TArray<UNiagaraSystem*> Systems;
 	Systems.Reserve(SystemAssets.Num());
@@ -645,7 +645,7 @@ void LoadAllSystemsInFolder(const TArray<FString>& Arguments)
 		TArray<FAssetData> SystemAssetsInFolder;
 		const FAssetRegistryModule& AssetRegistryModule = FModuleManager::LoadModuleChecked<FAssetRegistryModule>(TEXT("AssetRegistry"));
 		FARFilter Filter;
-		Filter.ClassNames.Add(UNiagaraSystem::StaticClass()->GetFName());
+		Filter.ClassPaths.Add(UNiagaraSystem::StaticClass()->GetClassPathName());
 		Filter.PackagePaths.Add(*Arguments[0]);
 		Filter.bRecursivePaths = true;
 		AssetRegistryModule.Get().GetAssets(Filter, SystemAssetsInFolder);
diff --git a/Engine/Plugins/FX/Niagara/Source/NiagaraEditor/Private/NiagaraEditorUtilities.cpp b/Engine/Plugins/FX/Niagara/Source/NiagaraEditor/Private/NiagaraEditorUtilities.cpp
--- a/Engine/Plugins/FX/Niagara/Source/NiagaraEditor/Private/NiagaraEditorUtilities.cpp
+++ b/Engine/Plugins/FX/Niagara/Source/NiagaraEditor/Private/NiagaraEditorUtilities.cpp
@@ -1374,9 +1374,9 @@ void FNiagaraEditorUtilities::PreprocessFunctionGraph(const UEdGraphSchema_Niaga
 void FNiagaraEditorUtilities::GetFilteredScriptAssets(FGetFilteredScriptAssetsOptions InFilter, TArray<FAssetData>& OutFilteredScriptAssets)
 {
 	FARFilter ScriptFilter;
-	ScriptFilter.ClassNames.Add(UNiagaraScript::StaticClass()->GetFName());
+	ScriptFilter.ClassPaths.Add(UNiagaraScript::StaticClass()->GetClassPathName());
 
-	const UEnum* NiagaraScriptUsageEnum = FindObjectChecked<UEnum>(ANY_PACKAGE, TEXT("ENiagaraScriptUsage"), true);
+	const UEnum* NiagaraScriptUsageEnum = FindObjectChecked<UEnum>(nullptr, TEXT("/Script/Niagara.ENiagaraScriptUsage"), true);
 	const FString QualifiedScriptUsageString = NiagaraScriptUsageEnum->GetNameStringByValue(static_cast<uint8>(InFilter.ScriptUsageToInclude));
 	int32 LastColonIndex;
 	QualifiedScriptUsageString.FindLastChar(TEXT(':'), LastColonIndex);
@@ -2706,7 +2706,7 @@ TArray<UNiagaraParameterDefinitions*> FNiagaraEditorUtilities::GetAllParameterDe
 
 	TArray<FAssetData> ParameterDefinitionsAssetData;
 	FAssetRegistryModule& AssetRegistryModule = FModuleManager::LoadModuleChecked<FAssetRegistryModule>("AssetRegistry");
-	AssetRegistryModule.GetRegistry().GetAssetsByClass(UNiagaraParameterDefinitions::StaticClass()->GetFName(), ParameterDefinitionsAssetData);
+	AssetRegistryModule.GetRegistry().GetAssetsByClass(UNiagaraParameterDefinitions::StaticClass()->GetClassPathName(), ParameterDefinitionsAssetData);
 	for (const FAssetData& ParameterDefinitionsAssetDatum : ParameterDefinitionsAssetData)
 	{
 		UNiagaraParameterDefinitions* ParameterDefinitions = Cast<UNiagaraParameterDefinitions>(ParameterDefinitionsAssetDatum.GetAsset());
@@ -2726,7 +2726,7 @@ bool FNiagaraEditorUtilities::GetAvailableParameterDefinitions(const TArray<FStr
 
 	// Gather asset registry filter args
 	FARFilter ARFilter;
-	ARFilter.ClassNames.Add(UNiagaraParameterDefinitions::StaticClass()->GetFName());
+	ARFilter.ClassPaths.Add(UNiagaraParameterDefinitions::StaticClass()->GetClassPathName());
 	ARFilter.bRecursivePaths = true;
 	ARFilter.bIncludeOnlyOnDiskAssets = true;
 
diff --git a/Engine/Plugins/FX/Niagara/Source/NiagaraEditor/Private/SNewSystemDialog.cpp b/Engine/Plugins/FX/Niagara/Source/NiagaraEditor/Private/SNewSystemDialog.cpp
--- a/Engine/Plugins/FX/Niagara/Source/NiagaraEditor/Private/SNewSystemDialog.cpp
+++ b/Engine/Plugins/FX/Niagara/Source/NiagaraEditor/Private/SNewSystemDialog.cpp
@@ -153,7 +153,7 @@ TOptional<FAssetData> SNewSystemDialog::GetSelectedSystemAsset() const
 	TArray<FAssetData> SelectedSystemAssets;
 	for (const FAssetData& SelectedAsset : AllSelectedAssets)
 	{
-		if (SelectedAsset.AssetClass == UNiagaraSystem::StaticClass()->GetFName())
+		if (SelectedAsset.AssetClassPath == UNiagaraSystem::StaticClass()->GetClassPathName())
 		{
 			SelectedSystemAssets.Add(SelectedAsset);
 		}
@@ -171,7 +171,7 @@ TArray<FAssetData> SNewSystemDialog::GetSelectedEmitterAssets() const
 	TArray<FAssetData> ConfirmedSelectedEmitterAssets;
 	for (const FAssetData& SelectedAsset : AllSelectedAssets)
 	{
-		if (SelectedAsset.AssetClass == UNiagaraEmitter::StaticClass()->GetFName())
+		if (SelectedAsset.AssetClassPath == UNiagaraEmitter::StaticClass()->GetClassPathName())
 		{
 			ConfirmedSelectedEmitterAssets.Add(SelectedAsset);
 		}
diff --git a/Engine/Plugins/FX/Niagara/Source/NiagaraEditor/Private/SNiagaraAssetPickerList.cpp b/Engine/Plugins/FX/Niagara/Source/NiagaraEditor/Private/SNiagaraAssetPickerList.cpp
--- a/Engine/Plugins/FX/Niagara/Source/NiagaraEditor/Private/SNiagaraAssetPickerList.cpp
+++ b/Engine/Plugins/FX/Niagara/Source/NiagaraEditor/Private/SNiagaraAssetPickerList.cpp
@@ -112,7 +112,7 @@ TArray<FAssetData> SNiagaraAssetPickerList::GetAssetDataForSelector(UClass* Asse
 {
 	FAssetRegistryModule& AssetRegistryModule = FModuleManager::LoadModuleChecked<FAssetRegistryModule>(TEXT("AssetRegistry"));
 	TArray<FAssetData> EmitterAssets;
-	AssetRegistryModule.Get().GetAssetsByClass(AssetClass->GetFName(), EmitterAssets);
+	AssetRegistryModule.Get().GetAssetsByClass(AssetClass->GetClassPathName(), EmitterAssets);
 
 	TArray<FAssetData> EmittersToShow;
 	if (TabOptions.GetOnlyShowTemplates())
diff --git a/Engine/Plugins/FX/Niagara/Source/NiagaraEditor/Private/Sequencer/NiagaraSequence/MovieSceneNiagaraEmitterTrack.cpp b/Engine/Plugins/FX/Niagara/Source/NiagaraEditor/Private/Sequencer/NiagaraSequence/MovieSceneNiagaraEmitterTrack.cpp
--- a/Engine/Plugins/FX/Niagara/Source/NiagaraEditor/Private/Sequencer/NiagaraSequence/MovieSceneNiagaraEmitterTrack.cpp
+++ b/Engine/Plugins/FX/Niagara/Source/NiagaraEditor/Private/Sequencer/NiagaraSequence/MovieSceneNiagaraEmitterTrack.cpp
@@ -283,7 +283,7 @@ void UMovieSceneNiagaraEmitterTrack::CreateSections(const FFrameRate& InFrameRes
 			FString* SectionClassName = FunctionNode->GetScriptData()->ScriptMetaData.Find("NiagaraTimelineSectionClass");
 			if (SectionClassName != nullptr)
 			{
-				UClass* SectionClass = FindObject<UClass>(ANY_PACKAGE, **SectionClassName);
+				UClass* SectionClass = UClass::TryFindTypeSlow<UClass>(*SectionClassName);
 				if (SectionClass != nullptr)
 				{
 					UMovieSceneNiagaraEmitterSectionBase* EmitterSection = nullptr;
diff --git a/Engine/Plugins/FX/Niagara/Source/NiagaraEditor/Private/TypeEditorUtilities/NiagaraDataInterfaceCurveTypeEditorUtilities.cpp b/Engine/Plugins/FX/Niagara/Source/NiagaraEditor/Private/TypeEditorUtilities/NiagaraDataInterfaceCurveTypeEditorUtilities.cpp
--- a/Engine/Plugins/FX/Niagara/Source/NiagaraEditor/Private/TypeEditorUtilities/NiagaraDataInterfaceCurveTypeEditorUtilities.cpp
+++ b/Engine/Plugins/FX/Niagara/Source/NiagaraEditor/Private/TypeEditorUtilities/NiagaraDataInterfaceCurveTypeEditorUtilities.cpp
@@ -54,7 +54,7 @@ TSharedRef<SWidget> FNiagaraDataInterfaceCurveTypeEditorUtilitiesBase::GetImport
 		AssetPickerConfig.OnAssetSelected = FOnAssetSelected::CreateRaw(this, &FNiagaraDataInterfaceCurveTypeEditorUtilitiesBase::CurveAssetSelected, CurveDataInterface, DataInterfaceChangedHandler);
 		AssetPickerConfig.bAllowNullSelection = false;
 		AssetPickerConfig.InitialAssetViewType = EAssetViewType::List;
-		AssetPickerConfig.Filter.ClassNames.Add(GetSupportedAssetClassName());
+		AssetPickerConfig.Filter.ClassPaths.Add(GetSupportedAssetClassName());
 	}
 
 	FContentBrowserModule& ContentBrowserModule = FModuleManager::Get().LoadModuleChecked<FContentBrowserModule>(TEXT("ContentBrowser"));
@@ -82,9 +82,9 @@ void FNiagaraDataInterfaceCurveTypeEditorUtilitiesBase::CurveAssetSelected(const
 	}
 }
 
-FName FNiagaraDataInterfaceCurveTypeEditorUtilities::GetSupportedAssetClassName() const
+FTopLevelAssetPath FNiagaraDataInterfaceCurveTypeEditorUtilities::GetSupportedAssetClassName() const
 {
-	return UCurveFloat::StaticClass()->GetFName();
+	return UCurveFloat::StaticClass()->GetClassPathName();
 }
 
 void FNiagaraDataInterfaceCurveTypeEditorUtilities::ImportSelectedAsset(UObject* SelectedAsset, UNiagaraDataInterfaceCurveBase* CurveDataInterface) const
@@ -99,9 +99,9 @@ void FNiagaraDataInterfaceCurveTypeEditorUtilities::ImportSelectedAsset(UObject*
 	}
 }
 
-FName FNiagaraDataInterfaceVectorCurveTypeEditorUtilities::GetSupportedAssetClassName() const
+FTopLevelAssetPath FNiagaraDataInterfaceVectorCurveTypeEditorUtilities::GetSupportedAssetClassName() const
 {
-	return UCurveVector::StaticClass()->GetFName();
+	return UCurveVector::StaticClass()->GetClassPathName();
 }
 
 void FNiagaraDataInterfaceVectorCurveTypeEditorUtilities::ImportSelectedAsset(UObject* SelectedAsset, UNiagaraDataInterfaceCurveBase* CurveDataInterface) const
@@ -118,9 +118,9 @@ void FNiagaraDataInterfaceVectorCurveTypeEditorUtilities::ImportSelectedAsset(UO
 	}
 }
 
-FName FNiagaraDataInterfaceColorCurveTypeEditorUtilities::GetSupportedAssetClassName() const
+FTopLevelAssetPath FNiagaraDataInterfaceColorCurveTypeEditorUtilities::GetSupportedAssetClassName() const
 {
-	return UCurveLinearColor::StaticClass()->GetFName();
+	return UCurveLinearColor::StaticClass()->GetClassPathName();
 }
 
 void FNiagaraDataInterfaceColorCurveTypeEditorUtilities::ImportSelectedAsset(UObject* SelectedAsset, UNiagaraDataInterfaceCurveBase* CurveDataInterface) const
diff --git a/Engine/Plugins/FX/Niagara/Source/NiagaraEditor/Private/TypeEditorUtilities/NiagaraDataInterfaceCurveTypeEditorUtilities.h b/Engine/Plugins/FX/Niagara/Source/NiagaraEditor/Private/TypeEditorUtilities/NiagaraDataInterfaceCurveTypeEditorUtilities.h
--- a/Engine/Plugins/FX/Niagara/Source/NiagaraEditor/Private/TypeEditorUtilities/NiagaraDataInterfaceCurveTypeEditorUtilities.h
+++ b/Engine/Plugins/FX/Niagara/Source/NiagaraEditor/Private/TypeEditorUtilities/NiagaraDataInterfaceCurveTypeEditorUtilities.h
@@ -17,7 +17,7 @@ public:
 
 protected:
 	/** Gets the class name used for filtering the curve asset picker. */
-	virtual FName GetSupportedAssetClassName() const = 0;
+	virtual FTopLevelAssetPath GetSupportedAssetClassName() const = 0;
 
 	/** Imports the selected curve asset into the supplied curve data interface. */
 	virtual void ImportSelectedAsset(UObject* SelectedAsset, UNiagaraDataInterfaceCurveBase* CurveDataInterface) const = 0;
@@ -36,7 +36,7 @@ class FNiagaraDataInterfaceCurveTypeEditorUtilities : public FNiagaraDataInterfa
 {
 protected:
 	//~ FNiagaraDataInterfaceCurveTypeEditorUtilitiesBase interface
-	virtual FName GetSupportedAssetClassName() const override;
+	virtual FTopLevelAssetPath GetSupportedAssetClassName() const override;
 	virtual void ImportSelectedAsset(UObject* SelectedAsset, UNiagaraDataInterfaceCurveBase* CurveDataInterface) const override;
 };
 
@@ -46,7 +46,7 @@ class FNiagaraDataInterfaceVectorCurveTypeEditorUtilities : public FNiagaraDataI
 {
 protected:
 	//~ FNiagaraDataInterfaceCurveTypeEditorUtilitiesBase interface
-	virtual FName GetSupportedAssetClassName() const override;
+	virtual FTopLevelAssetPath GetSupportedAssetClassName() const override;
 	virtual void ImportSelectedAsset(UObject* SelectedAsset, UNiagaraDataInterfaceCurveBase* CurveDataInterface) const override;
 };
 
@@ -56,6 +56,6 @@ class FNiagaraDataInterfaceColorCurveTypeEditorUtilities : public FNiagaraDataIn
 {
 protected:
 	//~ FNiagaraDataInterfaceCurveTypeEditorUtilitiesBase interface
-	virtual FName GetSupportedAssetClassName() const override;
+	virtual FTopLevelAssetPath GetSupportedAssetClassName() const override;
 	virtual void ImportSelectedAsset(UObject* SelectedAsset, UNiagaraDataInterfaceCurveBase* CurveDataInterface) const override;
 };
\ No newline at end of file
diff --git a/Engine/Plugins/FX/Niagara/Source/NiagaraEditor/Private/ViewModels/NiagaraEmitterViewModel.cpp b/Engine/Plugins/FX/Niagara/Source/NiagaraEditor/Private/ViewModels/NiagaraEmitterViewModel.cpp
--- a/Engine/Plugins/FX/Niagara/Source/NiagaraEditor/Private/ViewModels/NiagaraEmitterViewModel.cpp
+++ b/Engine/Plugins/FX/Niagara/Source/NiagaraEditor/Private/ViewModels/NiagaraEmitterViewModel.cpp
@@ -176,7 +176,7 @@ void FNiagaraEmitterViewModel::CreateNewParentWindow(TSharedRef<FNiagaraEmitterH
 	FAssetPickerConfig AssetPickerConfig;
 	AssetPickerConfig.SelectionMode = ESelectionMode::SingleToggle;
 	AssetPickerConfig.InitialAssetViewType = EAssetViewType::List;
-	AssetPickerConfig.Filter.ClassNames.Add(UNiagaraEmitter::StaticClass()->GetFName());
+	AssetPickerConfig.Filter.ClassPaths.Add(UNiagaraEmitter::StaticClass()->GetClassPathName());
 	AssetPickerConfig.OnAssetsActivated.BindSP(this, &FNiagaraEmitterViewModel::UpdateParentEmitter, EmitterHandleViewModel);
 
 	FContentBrowserModule& ContentBrowserModule = FModuleManager::Get().LoadModuleChecked<FContentBrowserModule>(TEXT("ContentBrowser"));
diff --git a/Engine/Plugins/FX/Niagara/Source/NiagaraEditor/Private/ViewModels/Stack/NiagaraStackFunctionInput.cpp b/Engine/Plugins/FX/Niagara/Source/NiagaraEditor/Private/ViewModels/Stack/NiagaraStackFunctionInput.cpp
--- a/Engine/Plugins/FX/Niagara/Source/NiagaraEditor/Private/ViewModels/Stack/NiagaraStackFunctionInput.cpp
+++ b/Engine/Plugins/FX/Niagara/Source/NiagaraEditor/Private/ViewModels/Stack/NiagaraStackFunctionInput.cpp
@@ -1503,7 +1503,7 @@ void UNiagaraStackFunctionInput::GetAvailableParameterHandles(TArray<FNiagaraPar
 	//Parameter Collections
 	FAssetRegistryModule& AssetRegistryModule = FModuleManager::LoadModuleChecked<FAssetRegistryModule>("AssetRegistry");
 	TArray<FAssetData> CollectionAssets;
-	AssetRegistryModule.Get().GetAssetsByClass(UNiagaraParameterCollection::StaticClass()->GetFName(), CollectionAssets);
+	AssetRegistryModule.Get().GetAssetsByClass(UNiagaraParameterCollection::StaticClass()->GetClassPathName(), CollectionAssets);
 
 	for (FAssetData& CollectionAsset : CollectionAssets)
 	{
diff --git a/Engine/Plugins/FX/Niagara/Source/NiagaraEditor/Private/ViewModels/Stack/NiagaraStackScriptItemGroup.cpp b/Engine/Plugins/FX/Niagara/Source/NiagaraEditor/Private/ViewModels/Stack/NiagaraStackScriptItemGroup.cpp
--- a/Engine/Plugins/FX/Niagara/Source/NiagaraEditor/Private/ViewModels/Stack/NiagaraStackScriptItemGroup.cpp
+++ b/Engine/Plugins/FX/Niagara/Source/NiagaraEditor/Private/ViewModels/Stack/NiagaraStackScriptItemGroup.cpp
@@ -986,7 +986,7 @@ TOptional<UNiagaraStackEntry::FDropRequestResponse> UNiagaraStackScriptItemGroup
 TOptional<UNiagaraStackEntry::FDropRequestResponse> UNiagaraStackScriptItemGroup::CanDropScriptsOnTarget(const UNiagaraStackEntry& TargetEntry, const FDropRequest& DropRequest)
 {
 	TSharedRef<const FNiagaraScriptDragOperation> DragDropOp = StaticCastSharedRef<const FNiagaraScriptDragOperation>(DropRequest.DragDropOperation);
-	const UEnum* NiagaraScriptUsageEnum = FindObjectChecked<UEnum>(ANY_PACKAGE, TEXT("ENiagaraScriptUsage"), true);
+	const UEnum* NiagaraScriptUsageEnum = FindObjectChecked<UEnum>(nullptr, TEXT("/Script/Niagara.ENiagaraScriptUsage"), true);
 
 	if (&TargetEntry != this && TargetEntry.IsA<UNiagaraStackModuleItem>() == false)
 	{
diff --git a/Engine/Plugins/FX/Niagara/Source/NiagaraEditor/Public/NiagaraEditorModule.h b/Engine/Plugins/FX/Niagara/Source/NiagaraEditor/Public/NiagaraEditorModule.h
--- a/Engine/Plugins/FX/Niagara/Source/NiagaraEditor/Public/NiagaraEditorModule.h
+++ b/Engine/Plugins/FX/Niagara/Source/NiagaraEditor/Public/NiagaraEditorModule.h
@@ -245,7 +245,7 @@ private:
 		{
 			const FAssetRegistryModule& AssetRegistryModule = FModuleManager::LoadModuleChecked<FAssetRegistryModule>(TEXT("AssetRegistry"));
 			TArray<FAssetData> AssetData;
-			AssetRegistryModule.GetRegistry().GetAssetsByClass(AssetType::StaticClass()->GetFName(), AssetData);
+			AssetRegistryModule.GetRegistry().GetAssetsByClass(AssetType::StaticClass()->GetClassPathName(), AssetData);
 
 			CachedAssets.Reset(AssetData.Num());
 			for (const FAssetData& AssetDatum : AssetData)
diff --git a/Engine/Plugins/FX/Niagara/Source/NiagaraEditorWidgets/Private/DetailCustomizations/NiagaraDataInterfaceCurveDetails.cpp b/Engine/Plugins/FX/Niagara/Source/NiagaraEditorWidgets/Private/DetailCustomizations/NiagaraDataInterfaceCurveDetails.cpp
--- a/Engine/Plugins/FX/Niagara/Source/NiagaraEditorWidgets/Private/DetailCustomizations/NiagaraDataInterfaceCurveDetails.cpp
+++ b/Engine/Plugins/FX/Niagara/Source/NiagaraEditorWidgets/Private/DetailCustomizations/NiagaraDataInterfaceCurveDetails.cpp
@@ -1516,13 +1516,13 @@ void FNiagaraDataInterfaceCurveDetailsBase::ImportSelectedAsset(UObject* Selecte
 
 TSharedRef<SWidget> FNiagaraDataInterfaceCurveDetailsBase::GetCurveToCopyMenu()
 {
-	FName ClassName = GetSupportedAssetClassName();
+	FTopLevelAssetPath ClassName = GetSupportedAssetClassName();
 	FAssetPickerConfig AssetPickerConfig;
 	{
 		AssetPickerConfig.OnAssetSelected = FOnAssetSelected::CreateSP(this, &FNiagaraDataInterfaceCurveDetailsBase::CurveToCopySelected);
 		AssetPickerConfig.bAllowNullSelection = false;
 		AssetPickerConfig.InitialAssetViewType = EAssetViewType::List;
-		AssetPickerConfig.Filter.ClassNames.Add(ClassName);
+		AssetPickerConfig.Filter.ClassPaths.Add(ClassName);
 	}
 
 	FContentBrowserModule& ContentBrowserModule = FModuleManager::Get().LoadModuleChecked<FContentBrowserModule>(TEXT("ContentBrowser"));
@@ -1553,9 +1553,9 @@ void FNiagaraDataInterfaceCurveDetails::GetCurveProperties(IDetailLayoutBuilder&
 	CurveProperties.Add(DetailBuilder.GetProperty(FName("Curve"), UNiagaraDataInterfaceCurve::StaticClass()));
 }
 
-FName FNiagaraDataInterfaceCurveDetails::GetSupportedAssetClassName() const
+FTopLevelAssetPath FNiagaraDataInterfaceCurveDetails::GetSupportedAssetClassName() const
 {
-	return UCurveFloat::StaticClass()->GetFName();
+	return UCurveFloat::StaticClass()->GetClassPathName();
 }
 
 void FNiagaraDataInterfaceCurveDetails::GetFloatCurvesFromAsset(UObject* SelectedAsset, TArray<FRichCurve>& FloatCurves) const
@@ -1576,9 +1576,9 @@ void FNiagaraDataInterfaceVector2DCurveDetails::GetCurveProperties(IDetailLayout
 	OutCurveProperties.Add(DetailBuilder.GetProperty(FName("YCurve"), UNiagaraDataInterfaceVector2DCurve::StaticClass()));
 }
 
-FName FNiagaraDataInterfaceVector2DCurveDetails::GetSupportedAssetClassName() const
+FTopLevelAssetPath FNiagaraDataInterfaceVector2DCurveDetails::GetSupportedAssetClassName() const
 {
-	return UCurveVector::StaticClass()->GetFName();
+	return UCurveVector::StaticClass()->GetClassPathName();
 }
 
 void FNiagaraDataInterfaceVector2DCurveDetails::GetFloatCurvesFromAsset(UObject* SelectedAsset, TArray<FRichCurve>& FloatCurves) const
@@ -1604,9 +1604,9 @@ void FNiagaraDataInterfaceVectorCurveDetails::GetCurveProperties(IDetailLayoutBu
 	OutCurveProperties.Add(DetailBuilder.GetProperty(FName("ZCurve"), UNiagaraDataInterfaceVectorCurve::StaticClass()));
 }
 
-FName FNiagaraDataInterfaceVectorCurveDetails::GetSupportedAssetClassName() const
+FTopLevelAssetPath FNiagaraDataInterfaceVectorCurveDetails::GetSupportedAssetClassName() const
 {
-	return UCurveVector::StaticClass()->GetFName();
+	return UCurveVector::StaticClass()->GetClassPathName();
 }
 
 void FNiagaraDataInterfaceVectorCurveDetails::GetFloatCurvesFromAsset(UObject* SelectedAsset, TArray<FRichCurve>& FloatCurves) const
@@ -1633,9 +1633,9 @@ void FNiagaraDataInterfaceVector4CurveDetails::GetCurveProperties(IDetailLayoutB
 	OutCurveProperties.Add(DetailBuilder.GetProperty(FName("WCurve"), UNiagaraDataInterfaceVector4Curve::StaticClass()));
 }
 
-FName FNiagaraDataInterfaceVector4CurveDetails::GetSupportedAssetClassName() const
+FTopLevelAssetPath FNiagaraDataInterfaceVector4CurveDetails::GetSupportedAssetClassName() const
 {
-	return UCurveLinearColor::StaticClass()->GetFName();
+	return UCurveLinearColor::StaticClass()->GetClassPathName();
 }
 
 void FNiagaraDataInterfaceVector4CurveDetails::GetFloatCurvesFromAsset(UObject* SelectedAsset, TArray<FRichCurve>& FloatCurves) const
@@ -1661,9 +1661,9 @@ void FNiagaraDataInterfaceColorCurveDetails::GetCurveProperties(IDetailLayoutBui
 	OutCurveProperties.Add(DetailBuilder.GetProperty(FName("AlphaCurve"), UNiagaraDataInterfaceColorCurve::StaticClass()));
 }
 
-FName FNiagaraDataInterfaceColorCurveDetails::GetSupportedAssetClassName() const
+FTopLevelAssetPath FNiagaraDataInterfaceColorCurveDetails::GetSupportedAssetClassName() const
 {
-	return UCurveLinearColor::StaticClass()->GetFName();
+	return UCurveLinearColor::StaticClass()->GetClassPathName();
 }
 
 void FNiagaraDataInterfaceColorCurveDetails::GetFloatCurvesFromAsset(UObject* SelectedAsset, TArray<FRichCurve>& FloatCurves) const
diff --git a/Engine/Plugins/FX/Niagara/Source/NiagaraEditorWidgets/Private/DetailCustomizations/NiagaraDataInterfaceCurveDetails.h b/Engine/Plugins/FX/Niagara/Source/NiagaraEditorWidgets/Private/DetailCustomizations/NiagaraDataInterfaceCurveDetails.h
--- a/Engine/Plugins/FX/Niagara/Source/NiagaraEditorWidgets/Private/DetailCustomizations/NiagaraDataInterfaceCurveDetails.h
+++ b/Engine/Plugins/FX/Niagara/Source/NiagaraEditorWidgets/Private/DetailCustomizations/NiagaraDataInterfaceCurveDetails.h
@@ -25,7 +25,7 @@ protected:
 	virtual bool GetIsColorCurve() const { return false; }
 	virtual float GetDefaultHeight() const { return 150; }
 	
-	virtual FName GetSupportedAssetClassName() const = 0;
+	virtual FTopLevelAssetPath GetSupportedAssetClassName() const = 0;
 	virtual void ImportSelectedAsset(UObject* SelectedAsset);
 	virtual void GetFloatCurvesFromAsset(UObject* SelectedAsset, TArray<FRichCurve>& FloatCurves) const = 0;
 	
@@ -52,7 +52,7 @@ public:
 protected:
 	virtual void GetCurveProperties(IDetailLayoutBuilder& DetailBuilder, TArray<TSharedRef<IPropertyHandle>>& OutCurveProperties) const override;
 	virtual float GetDefaultHeight() const override { return 130; }
-	virtual FName GetSupportedAssetClassName() const override;
+	virtual FTopLevelAssetPath GetSupportedAssetClassName() const override;
 	virtual void GetFloatCurvesFromAsset(UObject* SelectedAsset, TArray<FRichCurve>& FloatCurves) const override;
 };
 
@@ -65,7 +65,7 @@ public:
 
 protected:
 	virtual void GetCurveProperties(IDetailLayoutBuilder& DetailBuilder, TArray<TSharedRef<IPropertyHandle>>& OutCurveProperties) const override;
-	virtual FName GetSupportedAssetClassName() const override;
+	virtual FTopLevelAssetPath GetSupportedAssetClassName() const override;
 	virtual void GetFloatCurvesFromAsset(UObject* SelectedAsset, TArray<FRichCurve>& FloatCurves) const override;
 };
 
@@ -77,7 +77,7 @@ public:
 
 protected:
 	virtual void GetCurveProperties(IDetailLayoutBuilder& DetailBuilder, TArray<TSharedRef<IPropertyHandle>>& OutCurveProperties) const override;
-	virtual FName GetSupportedAssetClassName() const override;
+	virtual FTopLevelAssetPath GetSupportedAssetClassName() const override;
 	virtual void GetFloatCurvesFromAsset(UObject* SelectedAsset, TArray<FRichCurve>& FloatCurves) const override;
 };
 
@@ -89,7 +89,7 @@ public:
 
 protected:
 	virtual void GetCurveProperties(IDetailLayoutBuilder& DetailBuilder, TArray<TSharedRef<IPropertyHandle>>& OutCurveProperties) const override;
-	virtual FName GetSupportedAssetClassName() const override;
+	virtual FTopLevelAssetPath GetSupportedAssetClassName() const override;
 	virtual void GetFloatCurvesFromAsset(UObject* SelectedAsset, TArray<FRichCurve>& FloatCurves) const override;
 };
 
@@ -102,6 +102,6 @@ public:
 protected:
 	virtual void GetCurveProperties(IDetailLayoutBuilder& DetailBuilder, TArray<TSharedRef<IPropertyHandle>>& OutCurveProperties) const override;
 	virtual bool GetIsColorCurve() const override { return true; }
-	virtual FName GetSupportedAssetClassName() const override;
+	virtual FTopLevelAssetPath GetSupportedAssetClassName() const override;
 	virtual void GetFloatCurvesFromAsset(UObject* SelectedAsset, TArray<FRichCurve>& FloatCurves) const override;
 };
\ No newline at end of file
diff --git a/Engine/Plugins/Importers/AlembicImporter/Source/AlembicLibrary/Private/AbcFile.cpp b/Engine/Plugins/Importers/AlembicImporter/Source/AlembicLibrary/Private/AbcFile.cpp
--- a/Engine/Plugins/Importers/AlembicImporter/Source/AlembicLibrary/Private/AbcFile.cpp
+++ b/Engine/Plugins/Importers/AlembicImporter/Source/AlembicLibrary/Private/AbcFile.cpp
@@ -304,7 +304,7 @@ EAbcImportError FAbcFile::Import(UAbcImportSettings* InImportSettings)
 		FAssetRegistryModule& AssetRegistryModule = FModuleManager::LoadModuleChecked<FAssetRegistryModule>("AssetRegistry");
 		TArray<FAssetData> AssetData;
 		const UClass* Class = UMaterialInterface::StaticClass();
-		AssetRegistryModule.Get().GetAssetsByClass(Class->GetFName(), AssetData, true);
+		AssetRegistryModule.Get().GetAssetsByClass(Class->GetClassPathName(), AssetData, true);
 		for (FAbcPolyMesh* PolyMesh : PolyMeshes)
 		{
 			for (const FString& FaceSetName : PolyMesh->FaceSetNames)
diff --git a/Engine/Plugins/Media/MediaCompositing/Source/MediaCompositingEditor/Private/Sequencer/MediaPlayerPropertyTrackEditor.cpp b/Engine/Plugins/Media/MediaCompositing/Source/MediaCompositingEditor/Private/Sequencer/MediaPlayerPropertyTrackEditor.cpp
--- a/Engine/Plugins/Media/MediaCompositing/Source/MediaCompositingEditor/Private/Sequencer/MediaPlayerPropertyTrackEditor.cpp
+++ b/Engine/Plugins/Media/MediaCompositing/Source/MediaCompositingEditor/Private/Sequencer/MediaPlayerPropertyTrackEditor.cpp
@@ -156,7 +156,7 @@ TSharedPtr<SWidget> FMediaPlayerPropertyTrackEditor::BuildOutlinerEditWidget(con
 			AssetPickerConfig.bAllowNullSelection = false;
 			AssetPickerConfig.InitialAssetViewType = EAssetViewType::List;
 			AssetPickerConfig.Filter.bRecursiveClasses = true;
-			AssetPickerConfig.Filter.ClassNames.Add(UMediaSource::StaticClass()->GetFName());
+			AssetPickerConfig.Filter.ClassPaths.Add(UMediaSource::StaticClass()->GetClassPathName());
 			AssetPickerConfig.SaveSettingsName = TEXT("SequencerAssetPicker");
 			AssetPickerConfig.AdditionalReferencingAssets.Add(FAssetData(Sequence));
 		}
diff --git a/Engine/Plugins/Media/MediaCompositing/Source/MediaCompositingEditor/Private/Sequencer/MediaTrackEditor.cpp b/Engine/Plugins/Media/MediaCompositing/Source/MediaCompositingEditor/Private/Sequencer/MediaTrackEditor.cpp
--- a/Engine/Plugins/Media/MediaCompositing/Source/MediaCompositingEditor/Private/Sequencer/MediaTrackEditor.cpp
+++ b/Engine/Plugins/Media/MediaCompositing/Source/MediaCompositingEditor/Private/Sequencer/MediaTrackEditor.cpp
@@ -97,7 +97,7 @@ TSharedPtr<SWidget> FMediaTrackEditor::BuildOutlinerEditWidget(const FGuid& Obje
 			AssetPickerConfig.bAllowNullSelection = false;
 			AssetPickerConfig.InitialAssetViewType = EAssetViewType::List;
 			AssetPickerConfig.Filter.bRecursiveClasses = true;
-			AssetPickerConfig.Filter.ClassNames.Add(UMediaSource::StaticClass()->GetFName());
+			AssetPickerConfig.Filter.ClassPaths.Add(UMediaSource::StaticClass()->GetClassPathName());
 			AssetPickerConfig.SaveSettingsName = TEXT("SequencerAssetPicker");
 			AssetPickerConfig.AdditionalReferencingAssets.Add(FAssetData(Sequence));
 		}
diff --git a/Engine/Plugins/Media/MediaFrameworkUtilities/Source/MediaFrameworkUtilitiesEditor/Private/MediaFrameworkUtilitiesPlacement.cpp b/Engine/Plugins/Media/MediaFrameworkUtilities/Source/MediaFrameworkUtilitiesEditor/Private/MediaFrameworkUtilitiesPlacement.cpp
--- a/Engine/Plugins/Media/MediaFrameworkUtilities/Source/MediaFrameworkUtilitiesEditor/Private/MediaFrameworkUtilitiesPlacement.cpp
+++ b/Engine/Plugins/Media/MediaFrameworkUtilities/Source/MediaFrameworkUtilitiesEditor/Private/MediaFrameworkUtilitiesPlacement.cpp
@@ -190,7 +190,7 @@ public:
 		FAssetRegistryModule& AssetRegistryModule = FModuleManager::LoadModuleChecked<FAssetRegistryModule>(TEXT("AssetRegistry"));
 
 		TArray<FAssetData> AssetDatas;
-		AssetRegistryModule.Get().GetAssetsByClass("MediaBundle", AssetDatas, true);
+		AssetRegistryModule.Get().GetAssetsByClass(FTopLevelAssetPath(TEXT("/Script/MediaFrameworkUtilities"), TEXT("MediaBundle")), AssetDatas, true);
 		PlacementList.Reset(AssetDatas.Num());
 
 		for (const FAssetData& AssetData : AssetDatas)
diff --git a/Engine/Plugins/Media/MediaPlayerEditor/Source/MediaPlayerEditor/Private/Customizations/PlatformMediaSourceCustomization.cpp b/Engine/Plugins/Media/MediaPlayerEditor/Source/MediaPlayerEditor/Private/Customizations/PlatformMediaSourceCustomization.cpp
--- a/Engine/Plugins/Media/MediaPlayerEditor/Source/MediaPlayerEditor/Private/Customizations/PlatformMediaSourceCustomization.cpp
+++ b/Engine/Plugins/Media/MediaPlayerEditor/Source/MediaPlayerEditor/Private/Customizations/PlatformMediaSourceCustomization.cpp
@@ -210,7 +210,7 @@ FString FPlatformMediaSourceCustomization::HandleMediaSourceEntryBoxObjectPath(F
 bool FPlatformMediaSourceCustomization::HandleMediaSourceEntryBoxShouldFilterAsset(const FAssetData& AssetData)
 {
 	// Don't allow nesting platform media sources.
-	UClass* AssetClass = FindObject<UClass>(ANY_PACKAGE, *AssetData.AssetClass.ToString());
+	UClass* AssetClass = FindObject<UClass>(AssetData.AssetClassPath);
 	return AssetClass->IsChildOf(UPlatformMediaSource::StaticClass());
 }
 
diff --git a/Engine/Plugins/Media/MediaPlayerEditor/Source/MediaPlayerEditor/Private/Widgets/SMediaPlayerEditorMedia.cpp b/Engine/Plugins/Media/MediaPlayerEditor/Source/MediaPlayerEditor/Private/Widgets/SMediaPlayerEditorMedia.cpp
--- a/Engine/Plugins/Media/MediaPlayerEditor/Source/MediaPlayerEditor/Private/Widgets/SMediaPlayerEditorMedia.cpp
+++ b/Engine/Plugins/Media/MediaPlayerEditor/Source/MediaPlayerEditor/Private/Widgets/SMediaPlayerEditorMedia.cpp
@@ -39,8 +39,8 @@ void SMediaPlayerEditorMedia::Construct(const FArguments& InArgs, UMediaPlayer&
 	// initialize asset picker
 	FAssetPickerConfig AssetPickerConfig;
 	{
-		AssetPickerConfig.Filter.ClassNames.Add(UMediaPlaylist::StaticClass()->GetFName());
-		AssetPickerConfig.Filter.ClassNames.Add(UMediaSource::StaticClass()->GetFName());
+		AssetPickerConfig.Filter.ClassPaths.Add(UMediaPlaylist::StaticClass()->GetClassPathName());
+		AssetPickerConfig.Filter.ClassPaths.Add(UMediaSource::StaticClass()->GetClassPathName());
 		AssetPickerConfig.Filter.bRecursiveClasses = true;
 
 		AssetPickerConfig.AssetShowWarningText = LOCTEXT("NoMediaSourcesFound", "No media sources or play lists found.");
diff --git a/Engine/Plugins/Messaging/MessagingDebugger/Source/MessagingDebugger/Private/Models/MessagingDebuggerTypeFilter.h b/Engine/Plugins/Messaging/MessagingDebugger/Source/MessagingDebugger/Private/Models/MessagingDebuggerTypeFilter.h
--- a/Engine/Plugins/Messaging/MessagingDebugger/Source/MessagingDebugger/Private/Models/MessagingDebuggerTypeFilter.h
+++ b/Engine/Plugins/Messaging/MessagingDebugger/Source/MessagingDebugger/Private/Models/MessagingDebuggerTypeFilter.h
@@ -35,7 +35,7 @@ public:
 		
 		if (FilterString.ParseIntoArray(FilterSubstrings, TEXT(" "), true) > 0)
 		{
-			FString TypeName = TypeInfo->TypeName.ToString();
+			FString TypeName = TypeInfo->TypePathName.ToString();
 
 			for (int32 SubstringIndex = 0; SubstringIndex < FilterSubstrings.Num(); ++SubstringIndex)
 			{
diff --git a/Engine/Plugins/Messaging/MessagingDebugger/Source/MessagingDebugger/Private/Widgets/History/SMessagingHistoryTableRow.cpp b/Engine/Plugins/Messaging/MessagingDebugger/Source/MessagingDebugger/Private/Widgets/History/SMessagingHistoryTableRow.cpp
--- a/Engine/Plugins/Messaging/MessagingDebugger/Source/MessagingDebugger/Private/Widgets/History/SMessagingHistoryTableRow.cpp
+++ b/Engine/Plugins/Messaging/MessagingDebugger/Source/MessagingDebugger/Private/Widgets/History/SMessagingHistoryTableRow.cpp
@@ -104,7 +104,7 @@ TSharedRef<SWidget> SMessagingHistoryTableRow::GenerateWidgetForColumn(const FNa
 				SNew(STextBlock)
 					.ColorAndOpacity(this, &SMessagingHistoryTableRow::HandleTextColorAndOpacity)
 					.HighlightText(HighlightText)
-					.Text(FText::FromName(MessageInfo->Context->GetMessageType()))
+					.Text(FText::FromName(MessageInfo->Context->GetMessageTypePathName().GetAssetName()))
 			];
 	}
 	else if (ColumnName == "Recipients")
diff --git a/Engine/Plugins/Messaging/MessagingDebugger/Source/MessagingDebugger/Private/Widgets/Types/SMessagingTypesTableRow.cpp b/Engine/Plugins/Messaging/MessagingDebugger/Source/MessagingDebugger/Private/Widgets/Types/SMessagingTypesTableRow.cpp
--- a/Engine/Plugins/Messaging/MessagingDebugger/Source/MessagingDebugger/Private/Widgets/Types/SMessagingTypesTableRow.cpp
+++ b/Engine/Plugins/Messaging/MessagingDebugger/Source/MessagingDebugger/Private/Widgets/Types/SMessagingTypesTableRow.cpp
@@ -63,7 +63,7 @@ TSharedRef<SWidget> SMessagingTypesTableRow::GenerateWidgetForColumn(const FName
 			[
 				SNew(STextBlock)
 					.HighlightText(HighlightText)
-					.Text(FText::FromName(TypeInfo->TypeName))
+					.Text(FText::FromName(TypeInfo->TypePathName.GetAssetName()))
 			];
 	}
 	else if (ColumnName == "Visibility")
diff --git a/Engine/Plugins/Messaging/TcpMessaging/Source/TcpMessaging/Private/Transport/TcpDeserializedMessage.cpp b/Engine/Plugins/Messaging/TcpMessaging/Source/TcpMessaging/Private/Transport/TcpDeserializedMessage.cpp
--- a/Engine/Plugins/Messaging/TcpMessaging/Source/TcpMessaging/Private/Transport/TcpDeserializedMessage.cpp
+++ b/Engine/Plugins/Messaging/TcpMessaging/Source/TcpMessaging/Private/Transport/TcpDeserializedMessage.cpp
@@ -46,11 +46,11 @@ bool FTcpDeserializedMessage::Deserialize(const TSharedPtr<FArrayReader, ESPMode
 
 	// message type info
 	{
-		FName MessageType;
+		FTopLevelAssetPath MessageType;
 		MessageReader << MessageType;
 
 		// @todo gmp: cache message types for faster lookup
-		TypeInfoPtr = FindObjectSafe<UScriptStruct>(ANY_PACKAGE, *MessageType.ToString());
+		TypeInfoPtr = FindObjectSafe<UScriptStruct>(MessageType);
 
 		TypeInfo = TypeInfoPtr;
 		if (!TypeInfo.IsValid(false, true))
diff --git a/Engine/Plugins/Messaging/TcpMessaging/Source/TcpMessaging/Private/Transport/TcpMessageTransport.cpp b/Engine/Plugins/Messaging/TcpMessaging/Source/TcpMessaging/Private/Transport/TcpMessageTransport.cpp
--- a/Engine/Plugins/Messaging/TcpMessaging/Source/TcpMessaging/Private/Transport/TcpMessageTransport.cpp
+++ b/Engine/Plugins/Messaging/TcpMessaging/Source/TcpMessaging/Private/Transport/TcpMessageTransport.cpp
@@ -165,11 +165,11 @@ bool FTcpMessageTransport::TransportMessage(const TSharedRef<IMessageContext, ES
 
 	if (RecipientConnections.Num() == 0)
 	{
-		UE_LOG(LogTcpMessaging, Verbose, TEXT("No recipients for message %s from %s"), *Context->GetMessageType().ToString(), *Context->GetSender().ToString());
+		UE_LOG(LogTcpMessaging, Verbose, TEXT("No recipients for message %s from %s"), *Context->GetMessageTypePathName().ToString(), *Context->GetSender().ToString());
 		return false;
 	}
 
-	UE_LOG(LogTcpMessaging, Verbose, TEXT("Transporting message '%s' to %d connections"), *Context->GetMessageType().ToString(), RecipientConnections.Num());
+	UE_LOG(LogTcpMessaging, Verbose, TEXT("Transporting message '%s' to %d connections"), *Context->GetMessageTypePathName().ToString(), RecipientConnections.Num());
 
 	FTcpSerializedMessageRef SerializedMessage = MakeShareable(new FTcpSerializedMessage());
 
@@ -258,7 +258,7 @@ uint32 FTcpMessageTransport::Run()
 
 				while (Connection->Receive(Message, SenderNodeId))
 				{
-					UE_LOG(LogTcpMessaging, Verbose, TEXT("Received message '%s'"), *Message->GetMessageType().ToString());
+					UE_LOG(LogTcpMessaging, Verbose, TEXT("Received message '%s'"), *Message->GetMessageTypePathName().ToString());
 					TransportHandler->ReceiveTransportMessage(Message.ToSharedRef(), SenderNodeId);
 				}
 			}
diff --git a/Engine/Plugins/Messaging/TcpMessaging/Source/TcpMessaging/Private/Transport/TcpSerializeMessageTask.cpp b/Engine/Plugins/Messaging/TcpMessaging/Source/TcpMessaging/Private/Transport/TcpSerializeMessageTask.cpp
--- a/Engine/Plugins/Messaging/TcpMessaging/Source/TcpMessaging/Private/Transport/TcpSerializeMessageTask.cpp
+++ b/Engine/Plugins/Messaging/TcpMessaging/Source/TcpMessaging/Private/Transport/TcpSerializeMessageTask.cpp
@@ -20,8 +20,8 @@ void FTcpSerializeMessageTask::DoTask(ENamedThreads::Type CurrentThread, const F
 		// serialize context
 		FArchive& Archive = SerializedMessage.Get();
 		{
-			const FName& MessageType = MessageContext->GetMessageType();
-			Archive << const_cast<FName&>(MessageType);
+			const FTopLevelAssetPath& MessageType = MessageContext->GetMessageTypePathName();
+			Archive << const_cast<FTopLevelAssetPath&>(MessageType);
 
 			const FMessageAddress& Sender = MessageContext->GetSender();
 			Archive << const_cast<FMessageAddress&>(Sender);
diff --git a/Engine/Plugins/Messaging/UdpMessaging/Source/UdpMessaging/Private/Transport/UdpDeserializedMessage.cpp b/Engine/Plugins/Messaging/UdpMessaging/Source/UdpMessaging/Private/Transport/UdpDeserializedMessage.cpp
--- a/Engine/Plugins/Messaging/UdpMessaging/Source/UdpMessaging/Private/Transport/UdpDeserializedMessage.cpp
+++ b/Engine/Plugins/Messaging/UdpMessaging/Source/UdpMessaging/Private/Transport/UdpDeserializedMessage.cpp
@@ -44,6 +44,7 @@ class FUdpDeserializedMessageDetails
 public:
 	static bool DeserializeV10(FUdpDeserializedMessage& DeserializedMessage, FMemoryReader& MessageReader);
 	static bool DeserializeV11_15(FUdpDeserializedMessage& DeserializedMessage, FMemoryReader& MessageReader, bool bIsLWCBackwardCompatibilityMode);
+	static bool DeserializeV16(FUdpDeserializedMessage& DeserializedMessage, FMemoryReader& MessageReader);
 	static bool Deserialize(FUdpDeserializedMessage& DeserializedMessage, const FUdpReassembledMessage& ReassembledMessage);
 };
 
@@ -155,7 +156,7 @@ bool FUdpDeserializedMessageDetails::DeserializeV10(FUdpDeserializedMessage& Des
 		MessageReader << MessageType;
 
 		// @todo gmp: cache message types for faster lookup
-		DeserializedMessage.TypeInfo = FindObjectSafe<UScriptStruct>(ANY_PACKAGE, *MessageType.ToString());
+		DeserializedMessage.TypeInfo = UClass::TryFindTypeSlowSafe<UScriptStruct>(MessageType.ToString());
 
 		if (!DeserializedMessage.TypeInfo.IsValid(false, true))
 		{
@@ -241,7 +242,7 @@ bool FUdpDeserializedMessageDetails::DeserializeV11_15(FUdpDeserializedMessage&
 		MessageReader << MessageType;
 
 		// @todo gmp: cache message types for faster lookup
-		DeserializedMessage.TypeInfo = FindObjectSafe<UScriptStruct>(ANY_PACKAGE, *MessageType.ToString());
+		DeserializedMessage.TypeInfo = UClass::TryFindTypeSlowSafe<UScriptStruct>(MessageType.ToString());
 
 		if (!DeserializedMessage.TypeInfo.IsValid(false, true))
 		{
@@ -362,6 +363,136 @@ bool FUdpDeserializedMessageDetails::DeserializeV11_15(FUdpDeserializedMessage&
 	}
 }
 
+bool FUdpDeserializedMessageDetails::DeserializeV16(FUdpDeserializedMessage& DeserializedMessage, FMemoryReader& MessageReader)
+{
+	// message type info
+	{
+		FTopLevelAssetPath MessageType;
+		MessageReader << MessageType;
+
+		// @todo gmp: cache message types for faster lookup
+		DeserializedMessage.TypeInfo = FindObjectSafe<UScriptStruct>(MessageType);
+
+		if (!DeserializedMessage.TypeInfo.IsValid(false, true))
+		{
+			UE_LOG(LogUdpMessaging, Verbose, TEXT("No valid type info found for message type %s"), *MessageType.ToString());
+			return false;
+		}
+	}
+
+	// sender address
+	{
+		MessageReader << DeserializedMessage.Sender;
+	}
+
+	// recipient addresses
+	{
+		int32 NumRecipients = 0;
+		MessageReader << NumRecipients;
+
+		if ((NumRecipients < 0) || (NumRecipients > UDP_MESSAGING_MAX_RECIPIENTS))
+		{
+			return false;
+		}
+
+		DeserializedMessage.Recipients.Empty(NumRecipients);
+
+		while (0 < NumRecipients--)
+		{
+			MessageReader << DeserializedMessage.Recipients.AddDefaulted_GetRef();
+		}
+	}
+
+	// message scope
+	{
+		MessageReader << DeserializedMessage.Scope;
+
+		if (DeserializedMessage.Scope > EMessageScope::All)
+		{
+			return false;
+		}
+	}
+
+	// message flags
+	{
+		MessageReader << DeserializedMessage.Flags;
+	}
+
+	// time sent & expiration
+	{
+		MessageReader << DeserializedMessage.TimeSent;
+		MessageReader << DeserializedMessage.Expiration;
+	}
+
+	// annotations
+	{
+		int32 NumAnnotations = 0;
+		MessageReader << NumAnnotations;
+
+		if (NumAnnotations > UDP_MESSAGING_MAX_ANNOTATIONS)
+		{
+			return false;
+		}
+
+		while (0 < NumAnnotations--)
+		{
+			FName Key;
+			FString Value;
+
+			MessageReader << Key;
+			MessageReader << Value;
+
+			DeserializedMessage.Annotations.Add(Key, Value);
+		}
+	}
+
+	// wire format 
+	uint8 FormatId;
+	MessageReader << FormatId;
+	EUdpMessageFormat MessageFormat = (EUdpMessageFormat)FormatId;
+	
+	// create message body
+	DeserializedMessage.MessageData = FMemory::Malloc(DeserializedMessage.TypeInfo->GetStructureSize());
+	DeserializedMessage.TypeInfo->InitializeStruct(DeserializedMessage.MessageData);
+
+	constexpr bool bIsLWCBackwardCompatibilityMode = false;
+	switch (MessageFormat)
+	{
+	case EUdpMessageFormat::Json:
+	{
+		// deserialize json
+		FJsonStructDeserializerBackend Backend(MessageReader);
+		return FStructDeserializer::Deserialize(DeserializedMessage.MessageData, *DeserializedMessage.TypeInfo, Backend);
+	}
+	break;
+	case EUdpMessageFormat::CborPlatformEndianness:
+	{
+		// deserialize cbor (using this platform endianness).
+		FCborStructDeserializerBackend Backend(MessageReader, ECborEndianness::Platform, bIsLWCBackwardCompatibilityMode);
+		return FStructDeserializer::Deserialize(DeserializedMessage.MessageData, *DeserializedMessage.TypeInfo, Backend);
+	}
+	break;
+	case EUdpMessageFormat::CborStandardEndianness:
+	{
+		// deserialize cbor (using the CBOR standard endianness - big endian).
+		FCborStructDeserializerBackend Backend(MessageReader, ECborEndianness::StandardCompliant, bIsLWCBackwardCompatibilityMode);
+		return FStructDeserializer::Deserialize(DeserializedMessage.MessageData, *DeserializedMessage.TypeInfo, Backend);
+	}
+	break;
+	case EUdpMessageFormat::TaggedProperty:
+	{
+		// deserialize message body using tagged property
+		// Hack : this binary serialization should use a more standard protocol, should use cbor
+		DeserializedMessage.TypeInfo->SerializeItem(MessageReader, DeserializedMessage.MessageData, nullptr);
+		return !MessageReader.GetError();
+	}
+	break;
+	default:
+		// Unsupported format
+		return false;
+	}
+}
+
 bool FUdpDeserializedMessageDetails::Deserialize(FUdpDeserializedMessage& DeserializedMessage, const FUdpReassembledMessage& ReassembledMessage)
 {
 	// Note that some complex values are deserialized manually here, so that we
@@ -395,6 +526,11 @@ bool FUdpDeserializedMessageDetails::Deserialize(FUdpDeserializedMessage& Deseri
 		return DeserializeV11_15(DeserializedMessage, MessageReader, bIsLWCBackwardCompatibilityMode);
 		break;
 	}
+	case 16:
+	{
+		return DeserializeV16(DeserializedMessage, MessageReader);
+		break;
+	}
 
 	default:
 		UE_LOG(LogUdpMessaging, Error, TEXT("Unsupported Protocol Version message tasked for deserialization, discarding..."));
diff --git a/Engine/Plugins/Messaging/UdpMessaging/Source/UdpMessaging/Private/Transport/UdpMessageProcessor.cpp b/Engine/Plugins/Messaging/UdpMessaging/Source/UdpMessaging/Private/Transport/UdpMessageProcessor.cpp
--- a/Engine/Plugins/Messaging/UdpMessaging/Source/UdpMessaging/Private/Transport/UdpMessageProcessor.cpp
+++ b/Engine/Plugins/Messaging/UdpMessaging/Source/UdpMessaging/Private/Transport/UdpMessageProcessor.cpp
@@ -256,9 +256,10 @@ bool FUdpMessageProcessor::Init()
 		Beacon = new FUdpMessageBeacon(Socket, LocalNodeId, MulticastEndpoint);
 		SocketSender = new FUdpSocketSender(Socket, TEXT("FUdpMessageProcessor.Sender"));
 
-		// Current protocol version 15
+		// Current protocol version 16
 		SupportedProtocolVersions.Add(UDP_MESSAGING_TRANSPORT_PROTOCOL_VERSION);
-		// Support Protocol version 10, 11, 12, 13, 14
+		// Support Protocol version 10, 11, 12, 13, 14, 15
+		SupportedProtocolVersions.Add(15);
 		SupportedProtocolVersions.Add(14);
 		SupportedProtocolVersions.Add(13);
 		SupportedProtocolVersions.Add(12);
diff --git a/Engine/Plugins/Messaging/UdpMessaging/Source/UdpMessaging/Private/Transport/UdpMessageTransport.cpp b/Engine/Plugins/Messaging/UdpMessaging/Source/UdpMessaging/Private/Transport/UdpMessageTransport.cpp
--- a/Engine/Plugins/Messaging/UdpMessaging/Source/UdpMessaging/Private/Transport/UdpMessageTransport.cpp
+++ b/Engine/Plugins/Messaging/UdpMessaging/Source/UdpMessaging/Private/Transport/UdpMessageTransport.cpp
@@ -318,7 +318,7 @@ bool FUdpMessageTransport::TransportMessage(const TSharedRef<IMessageContext, ES
 	if (UE_GET_LOG_VERBOSITY(LogUdpMessaging) >= ELogVerbosity::Verbose)
 	{
 		FString RecipientStr = FString::JoinBy(Recipients, TEXT("+"), [](const FGuid& Guid) { return Guid.ToString(); });
-		UE_LOG(LogUdpMessaging, Verbose, TEXT("TransportMessage %s from %s to %s"), *Context->GetMessageType().ToString(), *Context->GetSender().ToString(), *RecipientStr);
+		UE_LOG(LogUdpMessaging, Verbose, TEXT("TransportMessage %s from %s to %s"), *Context->GetMessageTypePathName().ToString(), *Context->GetSender().ToString(), *RecipientStr);
 	}
 
 	return MessageProcessor->EnqueueOutboundMessage(Context, Recipients);
diff --git a/Engine/Plugins/Messaging/UdpMessaging/Source/UdpMessaging/Private/Transport/UdpSerializeMessageTask.cpp b/Engine/Plugins/Messaging/UdpMessaging/Source/UdpMessaging/Private/Transport/UdpSerializeMessageTask.cpp
--- a/Engine/Plugins/Messaging/UdpMessaging/Source/UdpMessaging/Private/Transport/UdpSerializeMessageTask.cpp
+++ b/Engine/Plugins/Messaging/UdpMessaging/Source/UdpMessaging/Private/Transport/UdpSerializeMessageTask.cpp
@@ -19,7 +19,7 @@ namespace UdpSerializeMessageTaskDetails
 /** Serialization Routine for message using Protocol version 10 */
 void SerializeMessageV10(FArchive& Archive, const TSharedRef<IMessageContext, ESPMode::ThreadSafe>& MessageContext)
 {
-	const FName& MessageType = MessageContext->GetMessageType();
+	const FName& MessageType = MessageContext->GetMessageTypePathName().GetAssetName();
 	Archive << const_cast<FName&>(MessageType);
 
 	const FMessageAddress& Sender = MessageContext->GetSender();
@@ -54,7 +54,7 @@ void SerializeMessageV10(FArchive& Archive, const TSharedRef<IMessageContext, ES
 /** Serialization Routine for message using Protocol version 11, 12, 13, 14 or 15. */
 void SerializeMessageV11_15(FArchive& Archive, const TSharedRef<IMessageContext, ESPMode::ThreadSafe>& MessageContext, EUdpMessageFormat MessageFormat, const EStructSerializerBackendFlags StructSerializerBackendFlags)
 {
-	const FName& MessageType = MessageContext->GetMessageType();
+	const FName& MessageType = MessageContext->GetMessageTypePathName().GetAssetName();
 	Archive << const_cast<FName&>(MessageType);
 
 	const FMessageAddress& Sender = MessageContext->GetSender();
@@ -94,6 +94,49 @@ void SerializeMessageV11_15(FArchive& Archive, const TSharedRef<IMessageContext,
 	FStructSerializer::Serialize(MessageContext->GetMessage(), *MessageContext->GetMessageTypeInfo(), Backend);
 }
 
+/** Serialization Routine for message using Protocol version 16. */
+void SerializeMessageV16(FArchive& Archive, const TSharedRef<IMessageContext, ESPMode::ThreadSafe>& MessageContext, EUdpMessageFormat MessageFormat, const EStructSerializerBackendFlags StructSerializerBackendFlags)
+{
+	const FTopLevelAssetPath& MessageType = MessageContext->GetMessageTypePathName();
+	Archive << const_cast<FTopLevelAssetPath&>(MessageType);
+
+	const FMessageAddress& Sender = MessageContext->GetSender();
+	Archive << const_cast<FMessageAddress&>(Sender);
+
+	const TArray<FMessageAddress>& Recipients = MessageContext->GetRecipients();
+	Archive << const_cast<TArray<FMessageAddress>&>(Recipients);
+
+	EMessageScope Scope = MessageContext->GetScope();
+	Archive << Scope;
+
+	EMessageFlags Flags = MessageContext->GetFlags();
+	Archive << Flags;
+
+	const FDateTime& TimeSent = MessageContext->GetTimeSent();
+	Archive << const_cast<FDateTime&>(TimeSent);
+
+	const FDateTime& Expiration = MessageContext->GetExpiration();
+	Archive << const_cast<FDateTime&>(Expiration);
+
+	int32 NumAnnotations = MessageContext->GetAnnotations().Num();
+	Archive << NumAnnotations;
+
+	for (const auto& AnnotationPair : MessageContext->GetAnnotations())
+	{
+		Archive << const_cast<FName&>(AnnotationPair.Key);
+		Archive << const_cast<FString&>(AnnotationPair.Value);
+	}
+
+	// Message Wire Format Id
+	check(MessageFormat == EUdpMessageFormat::CborPlatformEndianness || MessageFormat == EUdpMessageFormat::CborStandardEndianness); // Versions 11 to 15 only supports CBOR.
+	uint8 Format = (uint8)(MessageFormat);
+	Archive << Format;
+
+	// serialize message body with cbor
+	FCborStructSerializerBackend Backend(Archive, StructSerializerBackendFlags | (MessageFormat == EUdpMessageFormat::CborStandardEndianness ? EStructSerializerBackendFlags::WriteCborStandardEndianness : EStructSerializerBackendFlags::None));
+	FStructSerializer::Serialize(MessageContext->GetMessage(), *MessageContext->GetMessageTypeInfo(), Backend);
+}
+
 } // namespace UdpSerializeMessageTaskDetails
 
 
@@ -145,6 +188,11 @@ void FUdpSerializeMessageTask::DoTask(ENamedThreads::Type CurrentThread, const F
 				UdpSerializeMessageTaskDetails::SerializeMessageV11_15(Archive, MessageContext, SerializedMessage->GetFormat(), EStructSerializerBackendFlags::Default);
 				break;
 
+			case 16:
+				ProtocolMaxSegmentSize = UDP_MESSAGING_SEGMENT_SIZE * (int64)INT32_MAX;
+				UdpSerializeMessageTaskDetails::SerializeMessageV16(Archive, MessageContext, SerializedMessage->GetFormat(), EStructSerializerBackendFlags::Default);
+				break;
+
 			default:
 				// Unsupported protocol version
 				Serialized = false;
@@ -173,7 +221,7 @@ void FUdpSerializeMessageTask::DoTask(ENamedThreads::Type CurrentThread, const F
 		SerializedMessage->UpdateState(EUdpSerializedMessageState::Invalid);
 	}
 
-	UE_LOG(LogUdpMessaging, Verbose, TEXT("Serialized %s from %s to %d bytes"), *MessageContext->GetMessageType().ToString(), *MessageContext->GetSender().ToString(), SerializedMessage->TotalSize());
+	UE_LOG(LogUdpMessaging, Verbose, TEXT("Serialized %s from %s to %d bytes"), *MessageContext->GetMessageTypePathName().ToString(), *MessageContext->GetSender().ToString(), SerializedMessage->TotalSize());
 
 	// signal task completion
 	TSharedPtr<FEvent, ESPMode::ThreadSafe> CompletionEvent = CompletionEventPtr.Pin();
diff --git a/Engine/Plugins/Messaging/UdpMessaging/Source/UdpMessaging/Private/UdpMessagingPrivate.h b/Engine/Plugins/Messaging/UdpMessaging/Source/UdpMessaging/Private/UdpMessagingPrivate.h
--- a/Engine/Plugins/Messaging/UdpMessaging/Source/UdpMessaging/Private/UdpMessagingPrivate.h
+++ b/Engine/Plugins/Messaging/UdpMessaging/Source/UdpMessaging/Private/UdpMessagingPrivate.h
@@ -30,4 +30,4 @@ DECLARE_LOG_CATEGORY_EXTERN(LogUdpMessaging, Log, All);
  * Defines the protocol version of the UDP message transport.
  * @note When changing the version, ensure to update the serialization/deserialization code in UdpSerializeMessageTask.cpp/UdpDeserializedMessage.cpp and the supported versions in FUdpMessageProcessor::Init().
  */
-#define UDP_MESSAGING_TRANSPORT_PROTOCOL_VERSION 15
+#define UDP_MESSAGING_TRANSPORT_PROTOCOL_VERSION 16
diff --git a/Engine/Plugins/MovieScene/CustomizableSequencerTracks/Source/CustomizableSequencerTracksEditor/Private/SequencerTrackBPEditor.cpp b/Engine/Plugins/MovieScene/CustomizableSequencerTracks/Source/CustomizableSequencerTracksEditor/Private/SequencerTrackBPEditor.cpp
--- a/Engine/Plugins/MovieScene/CustomizableSequencerTracks/Source/CustomizableSequencerTracksEditor/Private/SequencerTrackBPEditor.cpp
+++ b/Engine/Plugins/MovieScene/CustomizableSequencerTracks/Source/CustomizableSequencerTracksEditor/Private/SequencerTrackBPEditor.cpp
@@ -31,7 +31,7 @@ namespace
 
 		TArray<FAssetData> ValidClasses;
 		{
-			FString DesiredClassName = FString::Printf(TEXT("Class'%s'"), *USequencerTrackBP::StaticClass()->GetPathName());
+			FString DesiredClassName = FString::Printf(TEXT("/Script/CoreUObject.Class'%s'"), *USequencerTrackBP::StaticClass()->GetPathName());
 
 			FARFilter Filter;
 			Filter.TagsAndValues.Add(FBlueprintTags::NativeParentClassPath, MoveTemp(DesiredClassName));
diff --git a/Engine/Plugins/MovieScene/LevelSequenceEditor/Source/LevelSequenceEditor/Private/LevelSequenceEditorModule.cpp b/Engine/Plugins/MovieScene/LevelSequenceEditor/Source/LevelSequenceEditor/Private/LevelSequenceEditorModule.cpp
--- a/Engine/Plugins/MovieScene/LevelSequenceEditor/Source/LevelSequenceEditor/Private/LevelSequenceEditorModule.cpp
+++ b/Engine/Plugins/MovieScene/LevelSequenceEditor/Source/LevelSequenceEditor/Private/LevelSequenceEditorModule.cpp
@@ -87,7 +87,7 @@ public:
 			}
 		};
 
-		FBlueprintAssetHandler::Get().RegisterHandler<FLevelSequenceAssetBlueprintHandler>(ULevelSequence::StaticClass()->GetFName());
+		FBlueprintAssetHandler::Get().RegisterHandler<FLevelSequenceAssetBlueprintHandler>(ULevelSequence::StaticClass()->GetClassPathName());
 	}
 	
 	virtual void ShutdownModule() override
diff --git a/Engine/Plugins/MovieScene/MatineeToLevelSequence/Source/MatineeToLevelSequence/Private/CameraAnimToTemplateSequenceConverter.cpp b/Engine/Plugins/MovieScene/MatineeToLevelSequence/Source/MatineeToLevelSequence/Private/CameraAnimToTemplateSequenceConverter.cpp
--- a/Engine/Plugins/MovieScene/MatineeToLevelSequence/Source/MatineeToLevelSequence/Private/CameraAnimToTemplateSequenceConverter.cpp
+++ b/Engine/Plugins/MovieScene/MatineeToLevelSequence/Source/MatineeToLevelSequence/Private/CameraAnimToTemplateSequenceConverter.cpp
@@ -335,7 +335,7 @@
 		// that it's going to be what we want.
 		for (const FAssetData& ExistingCameraAnimSequence : ExistingCameraAnimSequences)
 		{
-			if (ExistingCameraAnimSequence.AssetClass == UCameraAnimationSequence::StaticClass()->GetFName())
+			if (ExistingCameraAnimSequence.AssetClassPath == UCameraAnimationSequence::StaticClass()->GetClassPathName())
 			{
 				ExistingAsset = ExistingCameraAnimSequence.GetAsset();
 				if (ExistingAsset != nullptr)
diff --git a/Engine/Plugins/MovieScene/MatineeToLevelSequence/Source/MatineeToLevelSequence/Private/MatineeCameraShakeToNewCameraShakeConverter.cpp b/Engine/Plugins/MovieScene/MatineeToLevelSequence/Source/MatineeToLevelSequence/Private/MatineeCameraShakeToNewCameraShakeConverter.cpp
--- a/Engine/Plugins/MovieScene/MatineeToLevelSequence/Source/MatineeToLevelSequence/Private/MatineeCameraShakeToNewCameraShakeConverter.cpp
+++ b/Engine/Plugins/MovieScene/MatineeToLevelSequence/Source/MatineeToLevelSequence/Private/MatineeCameraShakeToNewCameraShakeConverter.cpp
@@ -40,14 +40,13 @@ namespace MovieScene
 /**
  * Returns whether the given blueprint has one of the given class names as its parent class.
  */
-bool HasAnyBlueprintParentClass(const FAssetData& AssetData, const TSet<FName>* ParentClassNames)
+bool HasAnyBlueprintParentClass(const FAssetData& AssetData, const TSet<FTopLevelAssetPath>* ParentClassNames)
 {
 	const FAssetDataTagMapSharedView::FFindTagResult FoundGeneratedClassTag = AssetData.TagsAndValues.FindTag(TEXT("GeneratedClass"));
 	if (FoundGeneratedClassTag.IsSet())
 	{
-		const FString GeneratedClassPath = FPackageName::ExportTextPathToObjectPath(FoundGeneratedClassTag.GetValue());
-		const FString GeneratedClassName = FPackageName::ObjectPathToObjectName(GeneratedClassPath);
-		if (ParentClassNames->Contains(*GeneratedClassName))
+		const FTopLevelAssetPath GeneratedClassPath(FPackageName::ExportTextPathToObjectPath(FoundGeneratedClassTag.GetValue()));
+		if (ParentClassNames->Contains(GeneratedClassPath))
 		{
 			return true;
 		}
@@ -76,7 +75,7 @@ class SMatineeCameraShakeConverterWidget : public SCompoundWidget
 {
 public:
 	SLATE_BEGIN_ARGS(SMatineeCameraShakeConverterWidget) {}
-	SLATE_ARGUMENT(TSet<FName>, MatineeCameraShakeClassNames)
+	SLATE_ARGUMENT(TSet<FTopLevelAssetPath>, MatineeCameraShakeClassNames)
 	SLATE_ARGUMENT(const FMatineeConverter*, MatineeConverter)
 	SLATE_END_ARGS()
 
@@ -92,7 +91,7 @@ public:
 
 		// Setup an asset picker that lists all Blueprint assets that have a known MatineeCameraShake parent class.
 		FARFilter Filter;
-		Filter.ClassNames.Add(UBlueprint::StaticClass()->GetFName());
+		Filter.ClassPaths.Add(UBlueprint::StaticClass()->GetClassPathName());
 		Filter.bRecursiveClasses = true;
 
 		FAssetPickerConfig AssetPickerConfig;
@@ -209,7 +208,7 @@ private:
 		Converter.ConvertMatineeCameraShakes(CurrentSelection);
 	}
 
-	TSet<FName> MatineeCameraShakeClassNames;
+	TSet<FTopLevelAssetPath> MatineeCameraShakeClassNames;
 	const FMatineeConverter* MatineeConverter = nullptr;
 
 	bool bFilterNonAnimCameraShakes = false;
@@ -231,9 +230,9 @@ TSharedRef<SWidget> FMatineeCameraShakeToNewCameraShakeConverter::CreateMatineeC
 
 	AssetRegistry.WaitForCompletion();
 
-	TArray<FName> ClassNames = { TEXT("CameraShake"), TEXT("MatineeCameraShake") };
-	TSet<FName> ExcludedClassNames;
-	TSet<FName> MatineeCameraShakeClassNames;
+	TArray<FTopLevelAssetPath> ClassNames = { FTopLevelAssetPath(TEXT("/Script/Engine"), TEXT("CameraShake")), FTopLevelAssetPath(TEXT("/Script/GameplayCameras"), TEXT("MatineeCameraShake")) };
+	TSet<FTopLevelAssetPath> ExcludedClassNames;
+	TSet<FTopLevelAssetPath> MatineeCameraShakeClassNames;
 	AssetRegistry.GetDerivedClassNames(ClassNames, ExcludedClassNames, MatineeCameraShakeClassNames);
 
 	return SNew(SMatineeCameraShakeConverterWidget)
diff --git a/Engine/Plugins/MovieScene/MovieRenderPipeline/Source/MovieRenderPipelineCore/Private/MoviePipelineUtils.cpp b/Engine/Plugins/MovieScene/MovieRenderPipeline/Source/MovieRenderPipelineCore/Private/MoviePipelineUtils.cpp
--- a/Engine/Plugins/MovieScene/MovieRenderPipeline/Source/MovieRenderPipelineCore/Private/MoviePipelineUtils.cpp
+++ b/Engine/Plugins/MovieScene/MovieRenderPipeline/Source/MovieRenderPipelineCore/Private/MoviePipelineUtils.cpp
@@ -37,8 +37,8 @@
 			TArray<FAssetData> ClassList;
 
 			FARFilter Filter;
-			Filter.ClassNames.Add(UMoviePipelineSetting::StaticClass()->GetFName());
-			Filter.ClassNames.Add(UBlueprint::StaticClass()->GetFName());
+			Filter.ClassPaths.Add(UMoviePipelineSetting::StaticClass()->GetClassPathName());
+			Filter.ClassPaths.Add(UBlueprint::StaticClass()->GetClassPathName());
 
 			// Include any Blueprint based objects as well, this includes things like Blutilities, UMG, and GameplayAbility objects
 			Filter.bRecursiveClasses = true;
@@ -60,11 +60,8 @@
 
 				if (!ParentClassName.IsEmpty())
 				{
-					UObject* Outer = nullptr;
-					ResolveName(Outer, ParentClassName, false, false);
+					ParentClass = UClass::TryFindTypeSlow<UClass>(FPackageName::ExportTextPathToObjectPath(ParentClassName));
 
-					ParentClass = FindObject<UClass>(ANY_PACKAGE, *ParentClassName);
-
 					if (!ParentClass->IsChildOf(UMoviePipelineSetting::StaticClass()))
 					{
 						continue;
diff --git a/Engine/Plugins/MovieScene/MovieRenderPipeline/Source/MovieRenderPipelineEditor/Private/Widgets/SMoviePipelineConfigPanel.cpp b/Engine/Plugins/MovieScene/MovieRenderPipeline/Source/MovieRenderPipelineEditor/Private/Widgets/SMoviePipelineConfigPanel.cpp
--- a/Engine/Plugins/MovieScene/MovieRenderPipeline/Source/MovieRenderPipelineEditor/Private/Widgets/SMoviePipelineConfigPanel.cpp
+++ b/Engine/Plugins/MovieScene/MovieRenderPipeline/Source/MovieRenderPipelineEditor/Private/Widgets/SMoviePipelineConfigPanel.cpp
@@ -367,7 +367,7 @@ TSharedRef<SWidget> SMoviePipelineConfigPanel::OnGeneratePresetsMenu()
 		AssetPickerConfig.bForceShowPluginContent = false;
 
 		AssetPickerConfig.AssetShowWarningText = LOCTEXT("NoPresets_Warning", "No Presets Found");
-		AssetPickerConfig.Filter.ClassNames.Add(ConfigAssetType->GetFName());
+		AssetPickerConfig.Filter.ClassPaths.Add(ConfigAssetType->GetClassPathName());
 		AssetPickerConfig.Filter.bRecursiveClasses = true;
 		AssetPickerConfig.OnAssetSelected = FOnAssetSelected::CreateSP(this, &SMoviePipelineConfigPanel::OnImportPreset);
 	}
@@ -394,7 +394,7 @@ bool SMoviePipelineConfigPanel::OpenSaveDialog(const FString& InDefaultPath, con
 	{
 		SaveAssetDialogConfig.DefaultPath = InDefaultPath;
 		SaveAssetDialogConfig.DefaultAssetName = InNewNameSuggestion;
-		SaveAssetDialogConfig.AssetClassNames.Add(ConfigAssetType->GetFName());
+		SaveAssetDialogConfig.AssetClassNames.Add(ConfigAssetType->GetClassPathName());
 		SaveAssetDialogConfig.ExistingAssetPolicy = ESaveAssetDialogExistingAssetPolicy::AllowButWarn;
 		SaveAssetDialogConfig.DialogTitleOverride = LOCTEXT("SaveConfigPresetDialogTitle", "Save Config Preset");
 	}
diff --git a/Engine/Plugins/MovieScene/MovieRenderPipeline/Source/MovieRenderPipelineEditor/Private/Widgets/SMoviePipelineQueueEditor.cpp b/Engine/Plugins/MovieScene/MovieRenderPipeline/Source/MovieRenderPipelineEditor/Private/Widgets/SMoviePipelineQueueEditor.cpp
--- a/Engine/Plugins/MovieScene/MovieRenderPipeline/Source/MovieRenderPipelineEditor/Private/Widgets/SMoviePipelineQueueEditor.cpp
+++ b/Engine/Plugins/MovieScene/MovieRenderPipeline/Source/MovieRenderPipelineEditor/Private/Widgets/SMoviePipelineQueueEditor.cpp
@@ -367,7 +367,7 @@ public:
 			AssetPickerConfig.SaveSettingsName = TEXT("MoviePipelineConfigAsset");
 
 			AssetPickerConfig.AssetShowWarningText = LOCTEXT("NoConfigs_Warning", "No Master Configurations Found");
-			AssetPickerConfig.Filter.ClassNames.Add(InClass->GetFName());
+			AssetPickerConfig.Filter.ClassPaths.Add(InClass->GetClassPathName());
 			AssetPickerConfig.OnAssetSelected = InOnAssetSelected;
 		}
 
@@ -1130,7 +1130,7 @@ TSharedRef<SWidget> SMoviePipelineQueueEditor::OnGenerateNewJobFromAssetMenu()
 		AssetPickerConfig.SaveSettingsName = TEXT("MoviePipelineQueueJobAsset");
 
 		AssetPickerConfig.AssetShowWarningText = LOCTEXT("NoSequences_Warning", "No Level Sequences Found");
-		AssetPickerConfig.Filter.ClassNames.Add(ULevelSequence::StaticClass()->GetFName());
+		AssetPickerConfig.Filter.ClassPaths.Add(ULevelSequence::StaticClass()->GetClassPathName());
 		AssetPickerConfig.OnAssetSelected = FOnAssetSelected::CreateSP(this, &SMoviePipelineQueueEditor::OnCreateJobFromAsset);
 	}
 
diff --git a/Engine/Plugins/MovieScene/MovieRenderPipeline/Source/MovieRenderPipelineEditor/Private/Widgets/SMoviePipelineQueuePanel.cpp b/Engine/Plugins/MovieScene/MovieRenderPipeline/Source/MovieRenderPipelineEditor/Private/Widgets/SMoviePipelineQueuePanel.cpp
--- a/Engine/Plugins/MovieScene/MovieRenderPipeline/Source/MovieRenderPipelineEditor/Private/Widgets/SMoviePipelineQueuePanel.cpp
+++ b/Engine/Plugins/MovieScene/MovieRenderPipeline/Source/MovieRenderPipelineEditor/Private/Widgets/SMoviePipelineQueuePanel.cpp
@@ -507,7 +507,7 @@ TSharedRef<SWidget> SMoviePipelineQueuePanel::OnGenerateSavedQueuesMenu()
 		AssetPickerConfig.bSortByPathInColumnView = false;
 
 		AssetPickerConfig.AssetShowWarningText = LOCTEXT("NoQueueAssets_Warning", "No Queues Found");
-		AssetPickerConfig.Filter.ClassNames.Add(UMoviePipelineQueue::StaticClass()->GetFName());
+		AssetPickerConfig.Filter.ClassPaths.Add(UMoviePipelineQueue::StaticClass()->GetClassPathName());
 		AssetPickerConfig.Filter.bRecursiveClasses = true;
 		AssetPickerConfig.OnAssetSelected = FOnAssetSelected::CreateSP(this, &SMoviePipelineQueuePanel::OnImportSavedQueueAssest);
 	}
@@ -534,7 +534,7 @@ bool SMoviePipelineQueuePanel::OpenSaveDialog(const FString& InDefaultPath, cons
 	{
 		SaveAssetDialogConfig.DefaultPath = InDefaultPath;
 		SaveAssetDialogConfig.DefaultAssetName = InNewNameSuggestion;
-		SaveAssetDialogConfig.AssetClassNames.Add(UMoviePipelineQueue::StaticClass()->GetFName());
+		SaveAssetDialogConfig.AssetClassNames.Add(UMoviePipelineQueue::StaticClass()->GetClassPathName());
 		SaveAssetDialogConfig.ExistingAssetPolicy = ESaveAssetDialogExistingAssetPolicy::AllowButWarn;
 		SaveAssetDialogConfig.DialogTitleOverride = LOCTEXT("SaveQueueAssetDialogTitle", "Save Queue Asset");
 	}
diff --git a/Engine/Plugins/MovieScene/TemplateSequence/Source/TemplateSequence/Private/TemplateSequence.cpp b/Engine/Plugins/MovieScene/TemplateSequence/Source/TemplateSequence/Private/TemplateSequence.cpp
--- a/Engine/Plugins/MovieScene/TemplateSequence/Source/TemplateSequence/Private/TemplateSequence.cpp
+++ b/Engine/Plugins/MovieScene/TemplateSequence/Source/TemplateSequence/Private/TemplateSequence.cpp
@@ -235,7 +235,7 @@
 
 	if (BoundActorClass.IsValid())
 	{
-		FAssetRegistryTag Tag("BoundActorClass", BoundActorClass->GetName(), FAssetRegistryTag::TT_Alphabetical);
+		FAssetRegistryTag Tag("BoundActorClass", BoundActorClass->GetPathName(), FAssetRegistryTag::TT_Alphabetical);
 		OutTags.Add(Tag);
 	}
 	else
@@ -245,6 +245,23 @@
 }
 
 #if WITH_EDITOR
+
+void UTemplateSequence::PostLoadAssetRegistryTags(const FAssetData& InAssetData, TArray<FAssetRegistryTag>& OutTagsAndValuesToUpdate) const
+{
+	Super::PostLoadAssetRegistryTags(InAssetData, OutTagsAndValuesToUpdate);
+
+	static const FName BoundActorClassTagName(TEXT("BoundActorClass"));
+	FString BoundActorClassTagValue = InAssetData.GetTagValueRef<FString>(BoundActorClassTagName);
+	if (!BoundActorClassTagValue.IsEmpty() && FPackageName::IsShortPackageName(BoundActorClassTagValue))
+	{
+		FTopLevelAssetPath BoundActorClassPathName = UClass::TryConvertShortTypeNameToPathName<UStruct>(BoundActorClassTagValue, ELogVerbosity::Warning, TEXT("UTemplateSequence::PostLoadAssetRegistryTags"));
+		if (!BoundActorClassPathName.IsNull())
+		{
+			OutTagsAndValuesToUpdate.Add(FAssetRegistryTag(BoundActorClassTagName, BoundActorClassPathName.ToString(), FAssetRegistryTag::TT_Alphabetical));
+		}
+	}
+}
+
 FText UTemplateSequence::GetDisplayName() const
 {
 	return UMovieSceneSequence::GetDisplayName();
diff --git a/Engine/Plugins/MovieScene/TemplateSequence/Source/TemplateSequence/Public/TemplateSequence.h b/Engine/Plugins/MovieScene/TemplateSequence/Source/TemplateSequence/Public/TemplateSequence.h
--- a/Engine/Plugins/MovieScene/TemplateSequence/Source/TemplateSequence/Public/TemplateSequence.h
+++ b/Engine/Plugins/MovieScene/TemplateSequence/Source/TemplateSequence/Public/TemplateSequence.h
@@ -47,6 +47,7 @@ public:
 	virtual void GetAssetRegistryTags(TArray<FAssetRegistryTag>& OutTags) const override;
 
 #if WITH_EDITOR
+	virtual void PostLoadAssetRegistryTags(const FAssetData& InAssetData, TArray<FAssetRegistryTag>& OutTagsAndValuesToUpdate) const;
 	virtual FText GetDisplayName() const override;
 
 	virtual ETrackSupport IsTrackSupported(TSubclassOf<class UMovieSceneTrack> InTrackClass) const override;
diff --git a/Engine/Plugins/MovieScene/TemplateSequence/Source/TemplateSequenceEditor/Private/TrackEditors/TemplateSequenceTrackEditor.cpp b/Engine/Plugins/MovieScene/TemplateSequence/Source/TemplateSequenceEditor/Private/TrackEditors/TemplateSequenceTrackEditor.cpp
--- a/Engine/Plugins/MovieScene/TemplateSequence/Source/TemplateSequenceEditor/Private/TrackEditors/TemplateSequenceTrackEditor.cpp
+++ b/Engine/Plugins/MovieScene/TemplateSequence/Source/TemplateSequenceEditor/Private/TrackEditors/TemplateSequenceTrackEditor.cpp
@@ -140,11 +140,11 @@ public:
 
 		// We will be looking for any template sequence whose root bound object is compatible with the current object binding.
 		// That means any template sequence bound to something of the same class, or a parent class.
-		TSet<FName> BaseClassNames;
+		TSet<FString> BaseClassNames;
 		const UClass* CurBaseClass = InBaseClass;
 		while (CurBaseClass)
 		{
-			BaseClassNames.Add(CurBaseClass->GetFName());
+			BaseClassNames.Add(CurBaseClass->GetPathName());
 			CurBaseClass = CurBaseClass->GetSuperClass();
 		}
 
@@ -159,24 +159,25 @@ public:
 			AssetPickerConfig.bAllowNullSelection = false;
 			AssetPickerConfig.InitialAssetViewType = EAssetViewType::List;
 			AssetPickerConfig.Filter.bRecursiveClasses = true;
-			AssetPickerConfig.Filter.ClassNames.Add(UTemplateSequence::StaticClass()->GetFName());
+			AssetPickerConfig.Filter.ClassPaths.Add(UTemplateSequence::StaticClass()->GetClassPathName());
 			AssetPickerConfig.SaveSettingsName = TEXT("SequencerAssetPicker");
 			AssetPickerConfig.AdditionalReferencingAssets.Add(FAssetData(Sequence));
 			if (LegacyBaseClass != nullptr)
 			{
-				AssetPickerConfig.Filter.ClassNames.Add(LegacyBaseClass->GetFName());
+				AssetPickerConfig.Filter.ClassPaths.Add(LegacyBaseClass->GetClassPathName());
 			}
 
 			AssetPickerConfig.OnShouldFilterAsset = FOnShouldFilterAsset::CreateLambda(
 				[this, BaseClassNames](const FAssetData& AssetData) -> bool
 				{
-					if (LegacyBaseClass == nullptr || AssetData.AssetClass != LegacyBaseClass->GetFName())
+					if (LegacyBaseClass == nullptr || AssetData.AssetClassPath != LegacyBaseClass->GetClassPathName())
 					{
 						const FAssetDataTagMapSharedView::FFindTagResult FoundBoundActorClass = AssetData.TagsAndValues.FindTag("BoundActorClass");
 						if (FoundBoundActorClass.IsSet())
 						{
 							// Filter this out if it's got an incompatible bound actor class.
-							const FName FoundBoundActorClassName(*FoundBoundActorClass.GetValue());
+							const FString FoundBoundActorClassName(FoundBoundActorClass.GetValue());
+							ensureAlwaysMsgf(!FPackageName::IsShortPackageName(FoundBoundActorClassName), TEXT("BoundActorClass tag contains short name \"%s\". Path name expected. See UTemplateSequence::GetAssetRegistryTags."), *FoundBoundActorClassName);
 							return !BaseClassNames.Contains(FoundBoundActorClassName);
 						}
 						else
diff --git a/Engine/Plugins/NetcodeUnitTest/NetcodeUnitTest/Source/NetcodeUnitTest/Private/NUTUtil.cpp b/Engine/Plugins/NetcodeUnitTest/NetcodeUnitTest/Source/NetcodeUnitTest/Private/NUTUtil.cpp
--- a/Engine/Plugins/NetcodeUnitTest/NetcodeUnitTest/Source/NetcodeUnitTest/Private/NUTUtil.cpp
+++ b/Engine/Plugins/NetcodeUnitTest/NetcodeUnitTest/Source/NetcodeUnitTest/Private/NUTUtil.cpp
@@ -119,7 +119,7 @@ bool NUTUtil::ParseEnum(const TCHAR* Stream, const TCHAR* Match, const TCHAR* En
 
 	if (FParse::Value(Stream, Match, StrValue) && StrValue.Len() > 0)
 	{
-		UEnum* EnumObj = FindObject<UEnum>(ANY_PACKAGE, EnumName);
+		UEnum* EnumObj = UClass::TryFindTypeSlow<UEnum>(EnumName);
 
 		if (EnumObj != nullptr)
 		{
diff --git a/Engine/Plugins/NetcodeUnitTest/NetcodeUnitTest/Source/NetcodeUnitTest/Private/NUTUtilReflectionParser.cpp b/Engine/Plugins/NetcodeUnitTest/NetcodeUnitTest/Source/NetcodeUnitTest/Private/NUTUtilReflectionParser.cpp
--- a/Engine/Plugins/NetcodeUnitTest/NetcodeUnitTest/Source/NetcodeUnitTest/Private/NUTUtilReflectionParser.cpp
+++ b/Engine/Plugins/NetcodeUnitTest/NetcodeUnitTest/Source/NetcodeUnitTest/Private/NUTUtilReflectionParser.cpp
@@ -352,7 +352,7 @@ TValueOrError<TSharedPtr<FVMReflection>, FExpressionError> FVMReflectionParser::
 						// Remove the FindObject tokens
 						Tokens.RemoveAt(0, FindTokenCount);
 
-						UClass* FindClass = ClassStr.Len() > 0 ? FindObject<UClass>(ANY_PACKAGE, *ClassStr) : UObject::StaticClass();
+						UClass* FindClass = ClassStr.Len() > 0 ? UClass::TryFindTypeSlow<UClass>(ClassStr) : UObject::StaticClass();
 						UObject* BestMatch = nullptr;
 
 						if (FindClass != nullptr)
diff --git a/Engine/Plugins/NetcodeUnitTest/NetcodeUnitTest/Source/NetcodeUnitTest/Public/NUTUtilReflection.h b/Engine/Plugins/NetcodeUnitTest/NetcodeUnitTest/Source/NetcodeUnitTest/Public/NUTUtilReflection.h
--- a/Engine/Plugins/NetcodeUnitTest/NetcodeUnitTest/Source/NetcodeUnitTest/Public/NUTUtilReflection.h
+++ b/Engine/Plugins/NetcodeUnitTest/NetcodeUnitTest/Source/NetcodeUnitTest/Public/NUTUtilReflection.h
@@ -856,7 +856,7 @@ public:
 	 * @param InFuncName	The name of the function
 	 */
 	FFuncReflection(const TCHAR* InClassName, const TCHAR* InFuncName)
-		: FFuncReflection(FindUField<UFunction>(FindObject<UClass>(ANY_PACKAGE, InClassName), InFuncName), InFuncName)
+		: FFuncReflection(FindUField<UFunction>(UClass::TryFindTypeSlow<UClass>(InClassName), InFuncName), InFuncName)
 	{
 	}
 
diff --git a/Engine/Plugins/Online/OnlineFramework/Source/Hotfix/Private/OnlineHotfixManager.cpp b/Engine/Plugins/Online/OnlineFramework/Source/Hotfix/Private/OnlineHotfixManager.cpp
--- a/Engine/Plugins/Online/OnlineFramework/Source/Hotfix/Private/OnlineHotfixManager.cpp
+++ b/Engine/Plugins/Online/OnlineFramework/Source/Hotfix/Private/OnlineHotfixManager.cpp
@@ -1046,7 +1046,7 @@ bool UOnlineHotfixManager::HotfixIniFile(const FString& FileName, const FString&
 					const FString ClassName = IniData.Mid(ClassNameStart, EndIndex - ClassNameStart);
 
 					// Look up the class to search for
-					UClass* ObjectClass = FindObject<UClass>(ANY_PACKAGE, *ClassName);
+					UClass* ObjectClass = UClass::TryFindTypeSlow<UClass>(ClassName);
 
 					if (ObjectClass)
 					{
@@ -1054,7 +1054,7 @@ bool UOnlineHotfixManager::HotfixIniFile(const FString& FileName, const FString&
 						const FString PerObjectName = IniData.Mid(StartIndex + 1, Count);
 
 						// Explicitly search the transient package (won't update non-transient objects)
-						UObject* PerObject = StaticFindObject(ObjectClass, ANY_PACKAGE, *PerObjectName, false);
+						UObject* PerObject = StaticFindFirstObject(ObjectClass, *PerObjectName, EFindFirstObjectOptions::NativeFirst);
 						if (PerObject != nullptr)
 						{
 							PerObjectConfigObjects.Add(PerObject);
@@ -1224,7 +1224,7 @@ bool UOnlineHotfixManager::IsMapLoaded(const FString& MapName)
 	FString MapPackageName(MapName.Left(MapName.Len() - 5));
 	MapPackageName = MapPackageName.Replace(*GameContentPath, TEXT("/Game"));
 	// If this map's UPackage exists, it is currently loaded
-	UPackage* MapPackage = FindObject<UPackage>(ANY_PACKAGE, *MapPackageName, true);
+	UPackage* MapPackage = FindObject<UPackage>(nullptr, *MapPackageName, true);
 	return MapPackage != nullptr;
 }
 
diff --git a/Engine/Plugins/Runtime/DataRegistry/Source/DataRegistry/Private/DataRegistrySource_DataTable.cpp b/Engine/Plugins/Runtime/DataRegistry/Source/DataRegistry/Private/DataRegistrySource_DataTable.cpp
--- a/Engine/Plugins/Runtime/DataRegistry/Source/DataRegistry/Private/DataRegistrySource_DataTable.cpp
+++ b/Engine/Plugins/Runtime/DataRegistry/Source/DataRegistry/Private/DataRegistrySource_DataTable.cpp
@@ -362,14 +362,14 @@ bool UMetaDataRegistrySource_DataTable::DoesAssetPassFilter(const FAssetData& As
 	if (AssetData.GetTagValue(RowStructureTagName, RowStructureString))
 	{
 		const UScriptStruct* ItemStruct = GetItemStruct();
-		if (RowStructureString == ItemStruct->GetName())
+		if (RowStructureString == ItemStruct->GetName() || RowStructureString == ItemStruct->GetStructPathName().ToString())
 		{
 			return true;
 		}
 		else
 		{
 			// TODO no 100% way to check for inherited row structs, but BP types can't inherit anyway
-			const UScriptStruct* RowStruct = FindObject<UScriptStruct>(ANY_PACKAGE, *RowStructureString, true);
+			const UScriptStruct* RowStruct = UClass::TryFindTypeSlow<UScriptStruct>(RowStructureString, EFindFirstObjectOptions::ExactClass);
 
 			// Check if the row struct is a child of the item struct
 			if (RowStruct != nullptr)
@@ -382,7 +382,7 @@ bool UMetaDataRegistrySource_DataTable::DoesAssetPassFilter(const FAssetData& As
 			// Otherwise check if the row struct has been redirected to the item struct
 			else
 			{
-				FName RowStructureName = FName(RowStructureString);
+				FName RowStructureName = FName(FPackageName::ObjectPathToObjectName(RowStructureString)); // RobM this should use path names as well
 				
 				TArray<FCoreRedirectObjectName> PreviousNames;
 				FCoreRedirects::FindPreviousNames(ECoreRedirectFlags::Type_Struct, ItemStruct->GetPathName(), PreviousNames);
diff --git a/Engine/Plugins/Runtime/DataRegistry/Source/DataRegistry/Private/DataRegistrySubsystem.cpp b/Engine/Plugins/Runtime/DataRegistry/Source/DataRegistry/Private/DataRegistrySubsystem.cpp
--- a/Engine/Plugins/Runtime/DataRegistry/Source/DataRegistry/Private/DataRegistrySubsystem.cpp
+++ b/Engine/Plugins/Runtime/DataRegistry/Source/DataRegistry/Private/DataRegistrySubsystem.cpp
@@ -659,6 +659,7 @@ void UDataRegistrySubsystem::ApplyPreregisterMap(UDataRegistry* Registry)
 
 	if (FoundPreregister)
 	{
+		const FTopLevelAssetPath ObjectClassPath(TEXT("/Script/CoreUObject"), TEXT("Object"));
 		for (int32 i = 0; i < FoundPreregister->Num(); i++)
 		{
 			bool bRegistered = false;
@@ -671,7 +672,7 @@ void UDataRegistrySubsystem::ApplyPreregisterMap(UDataRegistry* Registry)
 			else if (Settings->CanIgnoreMissingAssetData())
 			{
 				// Construct fake asset data and register that
-				AssetData = FAssetData(AssetPath.GetLongPackageName(), AssetPath.GetAssetPathString(), NAME_Object);
+				AssetData = FAssetData(AssetPath.GetLongPackageName(), AssetPath.GetAssetPathString(), ObjectClassPath);
 				bRegistered = Registry->RegisterSpecificAsset(AssetData, (*FoundPreregister)[i].Value);
 			}
 
diff --git a/Engine/Plugins/Runtime/GameplayAbilities/Source/GameplayAbilities/Public/Abilities/GameplayAbilityTypes.h b/Engine/Plugins/Runtime/GameplayAbilities/Source/GameplayAbilities/Public/Abilities/GameplayAbilityTypes.h
--- a/Engine/Plugins/Runtime/GameplayAbilities/Source/GameplayAbilities/Public/Abilities/GameplayAbilityTypes.h
+++ b/Engine/Plugins/Runtime/GameplayAbilities/Source/GameplayAbilities/Public/Abilities/GameplayAbilityTypes.h
@@ -438,13 +438,32 @@ DECLARE_MULTICAST_DELEGATE_OneParam(FGameplayAbilityEndedDelegate, const FAbilit
 /** Structure that tells AbilitySystemComponent what to bind to an InputComponent (see BindAbilityActivationToInputComponent) */
 struct FGameplayAbilityInputBinds
 {
+	UE_DEPRECATED(5.1, "Enum names are now represented by path names. Please use a version of FGameplayAbilityInputBinds constructor that accepts FTopLevelAssetPath.")
 	FGameplayAbilityInputBinds(FString InConfirmTargetCommand, FString InCancelTargetCommand, FString InEnumName, int32 InConfirmTargetInputID = INDEX_NONE, int32 InCancelTargetInputID = INDEX_NONE)
 		: ConfirmTargetCommand(InConfirmTargetCommand)
 		, CancelTargetCommand(InCancelTargetCommand)
+PRAGMA_DISABLE_DEPRECATION_WARNINGS
 		, EnumName(InEnumName)
+PRAGMA_ENABLE_DEPRECATION_WARNINGS
 		, ConfirmTargetInputID(InConfirmTargetInputID)
 		, CancelTargetInputID(InCancelTargetInputID)
-	{ }
+	{ 
+		TryFixShortEnumName();
+	}
+	FGameplayAbilityInputBinds(FString InConfirmTargetCommand, FString InCancelTargetCommand, FTopLevelAssetPath InEnumPathName, int32 InConfirmTargetInputID = INDEX_NONE, int32 InCancelTargetInputID = INDEX_NONE)
+		: ConfirmTargetCommand(InConfirmTargetCommand)
+		, CancelTargetCommand(InCancelTargetCommand)
+		, EnumPathName(InEnumPathName)
+		, ConfirmTargetInputID(InConfirmTargetInputID)
+		, CancelTargetInputID(InCancelTargetInputID)
+	{
+	}
+PRAGMA_DISABLE_DEPRECATION_WARNINGS
+	FGameplayAbilityInputBinds(FGameplayAbilityInputBinds&&) = default;
+	FGameplayAbilityInputBinds(const FGameplayAbilityInputBinds&) = default;
+	FGameplayAbilityInputBinds& operator=(FGameplayAbilityInputBinds&&) = default;
+	FGameplayAbilityInputBinds& operator=(const FGameplayAbilityInputBinds&) = default;
+PRAGMA_ENABLE_DEPRECATION_WARNINGS
 
 	/** Defines command string that will be bound to Confirm Targeting */
 	FString ConfirmTargetCommand;
@@ -453,15 +472,38 @@ struct FGameplayAbilityInputBinds
 	FString CancelTargetCommand;
 
 	/** Returns enum to use for ability binds. E.g., "Ability1"-"Ability9" input commands will be bound to ability activations inside the AbiltiySystemComponent */
+	UE_DEPRECATED(5.1, "Enum names are now represented by path names. Please use EnumPathName.")
 	FString	EnumName;
 
+	/** Returns enum to use for ability binds. E.g., "Ability1"-"Ability9" input commands will be bound to ability activations inside the AbiltiySystemComponent */
+	FTopLevelAssetPath EnumPathName;
+
 	/** If >=0, Confirm is bound to an entry in the enum */
 	int32 ConfirmTargetInputID;
 
 	/** If >=0, Cancel is bound to an entry in the enum */
 	int32 CancelTargetInputID;
 
-	UEnum* GetBindEnum() { return FindObject<UEnum>(ANY_PACKAGE, *EnumName); }
+	UEnum* GetBindEnum() 
+	{
+PRAGMA_DISABLE_DEPRECATION_WARNINGS
+		if (!EnumName.IsEmpty())
+		{
+			TryFixShortEnumName();
+		}
+PRAGMA_ENABLE_DEPRECATION_WARNINGS
+		return FindObject<UEnum>(EnumPathName);
+	}
+
+private:
+
+	void TryFixShortEnumName()
+	{
+PRAGMA_DISABLE_DEPRECATION_WARNINGS
+		EnumPathName = FTopLevelAssetPath(GetPathNameSafe(UClass::TryFindTypeSlow<UEnum>(EnumName)));
+		EnumName.Empty();
+PRAGMA_ENABLE_DEPRECATION_WARNINGS
+	}
 };
 
 
diff --git a/Engine/Plugins/Runtime/GameplayAbilities/Source/GameplayAbilitiesEditor/Private/AttributeDetails.cpp b/Engine/Plugins/Runtime/GameplayAbilities/Source/GameplayAbilitiesEditor/Private/AttributeDetails.cpp
--- a/Engine/Plugins/Runtime/GameplayAbilities/Source/GameplayAbilitiesEditor/Private/AttributeDetails.cpp
+++ b/Engine/Plugins/Runtime/GameplayAbilities/Source/GameplayAbilitiesEditor/Private/AttributeDetails.cpp
@@ -119,7 +119,7 @@ void FAttributePropertyDetails::OnChangeProperty(TSharedPtr<FString> ItemSelecte
 
 		FullString.Split( TEXT("."), &ClassName, &PropertyName);
 
-		UClass *FoundClass = FindObject<UClass>(ANY_PACKAGE, *ClassName);
+		UClass *FoundClass = UClass::TryFindTypeSlow<UClass>(ClassName);
 		if (FoundClass)
 		{
 			FProperty *Property = FindFProperty<FProperty>(FoundClass, *PropertyName);
diff --git a/Engine/Plugins/Runtime/MassEntity/Source/MassEntity/Private/MassEntityDebug.cpp b/Engine/Plugins/Runtime/MassEntity/Source/MassEntity/Private/MassEntityDebug.cpp
--- a/Engine/Plugins/Runtime/MassEntity/Source/MassEntity/Private/MassEntityDebug.cpp
+++ b/Engine/Plugins/Runtime/MassEntity/Source/MassEntity/Private/MassEntityDebug.cpp
@@ -7,7 +7,7 @@
 #include "Engine/Engine.h"
 #include "StructTypeBitSet.h"
 
-DEFINE_ENUM_TO_STRING(EMassProcessingPhase);
+DEFINE_ENUM_TO_STRING(EMassProcessingPhase, "/Script/MassEntity");
 
 namespace UE::Mass::Debug
 {
diff --git a/Engine/Plugins/Runtime/MassGameplay/Source/MassSpawner/Private/MassEntityTemplate.cpp b/Engine/Plugins/Runtime/MassGameplay/Source/MassSpawner/Private/MassEntityTemplate.cpp
--- a/Engine/Plugins/Runtime/MassGameplay/Source/MassSpawner/Private/MassEntityTemplate.cpp
+++ b/Engine/Plugins/Runtime/MassGameplay/Source/MassSpawner/Private/MassEntityTemplate.cpp
@@ -3,7 +3,7 @@
 #include "MassEntityTemplate.h"
 #include "VisualLogger/VisualLoggerTypes.h"
 
-DEFINE_ENUM_TO_STRING(EMassEntityTemplateIDType);
+DEFINE_ENUM_TO_STRING(EMassEntityTemplateIDType, "/Script/MassSpawner");
 
 //----------------------------------------------------------------------//
 //  FMassEntityTemplateID
diff --git a/Engine/Plugins/Runtime/Metasound/Source/MetasoundEditor/Private/MetasoundEditorGraphBuilder.cpp b/Engine/Plugins/Runtime/Metasound/Source/MetasoundEditor/Private/MetasoundEditorGraphBuilder.cpp
--- a/Engine/Plugins/Runtime/Metasound/Source/MetasoundEditor/Private/MetasoundEditorGraphBuilder.cpp
+++ b/Engine/Plugins/Runtime/Metasound/Source/MetasoundEditor/Private/MetasoundEditorGraphBuilder.cpp
@@ -906,7 +906,7 @@ namespace Metasound
 									{
 										if (AssetClass->IsChildOf(Class))
 										{
-											Filter.ClassNames.Add(Class->GetFName());
+											Filter.ClassPaths.Add(Class->GetClassPathName());
 											OutDefaultLiteral.Set(AssetObject);
 											bObjectFound = true;
 										}
diff --git a/Engine/Plugins/Runtime/Metasound/Source/MetasoundEditor/Private/MetasoundEditorGraphSchema.cpp b/Engine/Plugins/Runtime/Metasound/Source/MetasoundEditor/Private/MetasoundEditorGraphSchema.cpp
--- a/Engine/Plugins/Runtime/Metasound/Source/MetasoundEditor/Private/MetasoundEditorGraphSchema.cpp
+++ b/Engine/Plugins/Runtime/Metasound/Source/MetasoundEditor/Private/MetasoundEditorGraphSchema.cpp
@@ -1467,7 +1467,7 @@ void UMetasoundEditorGraphSchema::GetAssetsGraphHoverMessage(const TArray<FAsset
 	IMetasoundEditorModule& EditorModule = FModuleManager::GetModuleChecked<IMetasoundEditorModule>("MetaSoundEditor");
 	for (const FAssetData& Data : Assets)
 	{
-		if (!EditorModule.IsMetaSoundAssetClass(Data.GetClass()->GetFName()))
+		if (!EditorModule.IsMetaSoundAssetClass(Data.GetClass()->GetClassPathName()))
 		{
 			OutOkIcon = false;
 			OutTooltipText = TEXT("Asset(s) must all be MetaSounds.");
diff --git a/Engine/Plugins/Runtime/Metasound/Source/MetasoundEditor/Private/MetasoundEditorModule.cpp b/Engine/Plugins/Runtime/Metasound/Source/MetasoundEditor/Private/MetasoundEditorModule.cpp
--- a/Engine/Plugins/Runtime/Metasound/Source/MetasoundEditor/Private/MetasoundEditorModule.cpp
+++ b/Engine/Plugins/Runtime/Metasound/Source/MetasoundEditor/Private/MetasoundEditorModule.cpp
@@ -232,7 +232,7 @@ namespace Metasound
 			{
 				using namespace Frontend;
 
-				if (!IsMetaSoundAssetClass(InAssetData.AssetClass))
+				if (!IsMetaSoundAssetClass(InAssetData.AssetClassPath))
 				{
 					return;
 				}
@@ -259,7 +259,7 @@ namespace Metasound
 			{
 				using namespace Frontend;
 
-				if (!IsMetaSoundAssetClass(InAssetData.AssetClass))
+				if (!IsMetaSoundAssetClass(InAssetData.AssetClassPath))
 				{
 					return;
 				}
@@ -299,7 +299,7 @@ namespace Metasound
 				{
 					if (UObject* Obj = Pair.Key)
 					{
-						if (IsMetaSoundAssetClass(Obj->GetClass()->GetFName()))
+						if (IsMetaSoundAssetClass(Obj->GetClass()->GetClassPathName()))
 						{
 							check(GEngine);
 							UMetaSoundAssetSubsystem* AssetSubsystem = GEngine->GetEngineSubsystem<UMetaSoundAssetSubsystem>();
@@ -313,7 +313,7 @@ namespace Metasound
 
 					if (UObject* Obj = Pair.Value)
 					{
-						if (IsMetaSoundAssetClass(Obj->GetClass()->GetFName()))
+						if (IsMetaSoundAssetClass(Obj->GetClass()->GetClassPathName()))
 						{
 							check(GEngine);
 							UMetaSoundAssetSubsystem* AssetSubsystem = GEngine->GetEngineSubsystem<UMetaSoundAssetSubsystem>();
@@ -348,7 +348,7 @@ namespace Metasound
 
 			void RemoveAssetFromClassRegistry(const FAssetData& InAssetData)
 			{
-				if (IsMetaSoundAssetClass(InAssetData.AssetClass))
+				if (IsMetaSoundAssetClass(InAssetData.AssetClassPath))
 				{
 					check(GEngine);
 					UMetaSoundAssetSubsystem* AssetSubsystem = GEngine->GetEngineSubsystem<UMetaSoundAssetSubsystem>();
@@ -365,7 +365,7 @@ namespace Metasound
 
 			void RenameAssetInClassRegistry(const FAssetData& InAssetData, const FString& InOldObjectPath)
 			{
-				if (IsMetaSoundAssetClass(InAssetData.AssetClass))
+				if (IsMetaSoundAssetClass(InAssetData.AssetClassPath))
 				{
 					check(GEngine);
 					UMetaSoundAssetSubsystem* AssetSubsystem = GEngine->GetEngineSubsystem<UMetaSoundAssetSubsystem>();
@@ -563,7 +563,7 @@ namespace Metasound
 					AssetPrimeStatus = EAssetPrimeStatus::InProgress;
 
 					FARFilter Filter;
-					Filter.ClassNames = MetaSoundClassNames;
+					Filter.ClassPaths = MetaSoundClassNames;
 
 					FAssetRegistryModule& AssetRegistryModule = FModuleManager::LoadModuleChecked<FAssetRegistryModule>("AssetRegistry");
 					AssetRegistryModule.Get().EnumerateAssets(Filter, [this](const FAssetData& AssetData)
@@ -643,7 +643,7 @@ namespace Metasound
 				return PinTypes.Find(InDataTypeName);
 			}
 
-			virtual bool IsMetaSoundAssetClass(const FName InClassName) const override
+			virtual bool IsMetaSoundAssetClass(const FTopLevelAssetPath& InClassName) const override
 			{
 				// TODO: Move to IMetasoundUObjectRegistry (overload IsRegisteredClass to take in class name?)
 				return MetaSoundClassNames.Contains(InClassName);
@@ -721,8 +721,8 @@ namespace Metasound
 					GetMutableDefault<UMetasoundEditorSettings>()
 				);
 
-				MetaSoundClassNames.Add(UMetaSound::StaticClass()->GetFName());
-				MetaSoundClassNames.Add(UMetaSoundSource::StaticClass()->GetFName());
+				MetaSoundClassNames.Add(UMetaSound::StaticClass()->GetClassPathName());
+				MetaSoundClassNames.Add(UMetaSoundSource::StaticClass()->GetClassPathName());
 
 				FAssetTypeActions_MetaSound::RegisterMenuActions();
 				FAssetTypeActions_MetaSoundSource::RegisterMenuActions();
@@ -794,7 +794,7 @@ namespace Metasound
 				MetaSoundClassNames.Reset();
 			}
 
-			TArray<FName> MetaSoundClassNames;
+			TArray<FTopLevelAssetPath> MetaSoundClassNames;
 
 			TArray<TSharedPtr<FAssetTypeActions_Base>> AssetActions;
 			TMap<EMetasoundFrontendLiteralType, const TSubclassOf<UMetasoundEditorGraphMemberDefaultLiteral>> InputDefaultLiteralClassRegistry;
diff --git a/Engine/Plugins/Runtime/Metasound/Source/MetasoundEditor/Public/MetasoundEditorModule.h b/Engine/Plugins/Runtime/Metasound/Source/MetasoundEditor/Public/MetasoundEditorModule.h
--- a/Engine/Plugins/Runtime/Metasound/Source/MetasoundEditor/Public/MetasoundEditorModule.h
+++ b/Engine/Plugins/Runtime/Metasound/Source/MetasoundEditor/Public/MetasoundEditorModule.h
@@ -82,7 +82,7 @@ namespace Metasound
 			// By default, proxy classes support child classes & inheritance.
 			virtual void RegisterExplicitProxyClass(const UClass& InClass) = 0;
 
-			virtual bool IsMetaSoundAssetClass(const FName InClassName) const = 0;
+			virtual bool IsMetaSoundAssetClass(const FTopLevelAssetPath& InClassName) const = 0;
 
 			// Primes MetaSound assets, effectively loading the asset asynchronously (if not already
 			// loaded) & registers them if not already registered with the MetaSound Class Registry.
diff --git a/Engine/Plugins/Runtime/ReplicationGraph/Source/Private/ReplicationGraphDebugging.cpp b/Engine/Plugins/Runtime/ReplicationGraph/Source/Private/ReplicationGraphDebugging.cpp
--- a/Engine/Plugins/Runtime/ReplicationGraph/Source/Private/ReplicationGraphDebugging.cpp
+++ b/Engine/Plugins/Runtime/ReplicationGraph/Source/Private/ReplicationGraphDebugging.cpp
@@ -511,7 +511,7 @@ FAutoConsoleCommandWithWorldAndArgs NetRepGraphSetClassCullDistance(TEXT("Net.Re
 			return;
 		}
 
-		UClass* Class = FindObject<UClass>(ANY_PACKAGE, *Args[0]);
+		UClass* Class = UClass::TryFindTypeSlow<UClass>(Args[0]);
 		if (Class == nullptr)
 		{
 			UE_LOG(LogReplicationGraph, Display, TEXT("Could not find Class: %s"), *Args[0]);
@@ -572,7 +572,7 @@ FAutoConsoleCommandWithWorldAndArgs NetRepGraphSetPeriodFrame(TEXT("Net.RepGraph
 			return;
 		}
 
-		UClass* Class = FindObject<UClass>(ANY_PACKAGE, *Args[0]);
+		UClass* Class = UClass::TryFindTypeSlow<UClass>(Args[0]);
 		if (Class == nullptr)
 		{
 			UE_LOG(LogReplicationGraph, Display, TEXT("Could not find Class: %s"), *Args[0]);
diff --git a/Engine/Plugins/Runtime/StateTree/Source/StateTreeEditorModule/Private/Commandlets/StateTreeCompileAllCommandlet.cpp b/Engine/Plugins/Runtime/StateTree/Source/StateTreeEditorModule/Private/Commandlets/StateTreeCompileAllCommandlet.cpp
--- a/Engine/Plugins/Runtime/StateTree/Source/StateTreeEditorModule/Private/Commandlets/StateTreeCompileAllCommandlet.cpp
+++ b/Engine/Plugins/Runtime/StateTree/Source/StateTreeEditorModule/Private/Commandlets/StateTreeCompileAllCommandlet.cpp
@@ -1,4 +1,4 @@
-﻿// Copyright Epic Games, Inc. All Rights Reserved.
+// Copyright Epic Games, Inc. All Rights Reserved.
 
 #include "Commandlets/StateTreeCompileAllCommandlet.h"
 #include "StateTree.h"
@@ -42,7 +42,7 @@ int32 UStateTreeCompileAllCommandlet::Main(const FString& Params)
 	
 	UE_LOG(LogStateTreeCompile, Display, TEXT("Gathering All StateTrees From Asset Registry..."));
 	TArray<FAssetData> StateTreeAssetList;
-	AssetRegistryModule.Get().GetAssetsByClass(UStateTree::StaticClass()->GetFName(), StateTreeAssetList, /*bSearchSubClasses*/false);
+	AssetRegistryModule.Get().GetAssetsByClass(UStateTree::StaticClass()->GetClassPathName(), StateTreeAssetList, /*bSearchSubClasses*/false);
 
 	int32 Counter = 0;
 	for (const FAssetData& Asset : StateTreeAssetList)
diff --git a/Engine/Plugins/Runtime/StateTree/Source/StateTreeEditorModule/Private/Customizations/StateTreeEditorNodeDetails.cpp b/Engine/Plugins/Runtime/StateTree/Source/StateTreeEditorModule/Private/Customizations/StateTreeEditorNodeDetails.cpp
--- a/Engine/Plugins/Runtime/StateTree/Source/StateTreeEditorModule/Private/Customizations/StateTreeEditorNodeDetails.cpp
+++ b/Engine/Plugins/Runtime/StateTree/Source/StateTreeEditorModule/Private/Customizations/StateTreeEditorNodeDetails.cpp
@@ -187,10 +187,10 @@ void FStateTreeEditorNodeDetails::CustomizeHeader(TSharedRef<class IPropertyHand
 	static const FName BaseClassMetaName(TEXT("BaseClass")); // TODO: move these names into one central place.
 	
 	const FString BaseStructName = StructProperty->GetMetaData(BaseStructMetaName);
-	BaseScriptStruct = FindObject<UScriptStruct>(ANY_PACKAGE, *BaseStructName);
+	BaseScriptStruct = UClass::TryFindTypeSlow<UScriptStruct>(BaseStructName);
 
 	const FString BaseClassName = StructProperty->GetMetaData(BaseClassMetaName);
-	BaseClass = FindObject<UClass>(ANY_PACKAGE, *BaseClassName);
+	BaseClass = UClass::TryFindTypeSlow<UClass>(BaseClassName);
 	
 	const FIsResetToDefaultVisible IsResetVisible = FIsResetToDefaultVisible::CreateSP(this, &FStateTreeEditorNodeDetails::ShouldResetToDefault);
 	const FResetToDefaultHandler ResetHandler = FResetToDefaultHandler::CreateSP(this, &FStateTreeEditorNodeDetails::ResetToDefault);
diff --git a/Engine/Plugins/Runtime/StateTree/Source/StateTreeEditorModule/Private/StateTreeNodeClassCache.cpp b/Engine/Plugins/Runtime/StateTree/Source/StateTreeEditorModule/Private/StateTreeNodeClassCache.cpp
--- a/Engine/Plugins/Runtime/StateTree/Source/StateTreeEditorModule/Private/StateTreeNodeClassCache.cpp
+++ b/Engine/Plugins/Runtime/StateTree/Source/StateTreeEditorModule/Private/StateTreeNodeClassCache.cpp
@@ -263,7 +263,7 @@ void FStateTreeNodeClassCache::CacheClasses()
 	TArray<FAssetData> BlueprintList;
 
 	FARFilter Filter;
-	Filter.ClassNames.Add(UBlueprint::StaticClass()->GetFName());
+	Filter.ClassPaths.Add(UBlueprint::StaticClass()->GetClassPathName());
 	AssetRegistryModule.Get().GetAssets(Filter, BlueprintList);
 
 	for (int32 i = 0; i < BlueprintList.Num(); i++)
diff --git a/Engine/Plugins/Runtime/StateTree/Source/StateTreeModule/Private/StateTree.cpp b/Engine/Plugins/Runtime/StateTree/Source/StateTreeModule/Private/StateTree.cpp
--- a/Engine/Plugins/Runtime/StateTree/Source/StateTreeModule/Private/StateTree.cpp
+++ b/Engine/Plugins/Runtime/StateTree/Source/StateTreeModule/Private/StateTree.cpp
@@ -38,11 +38,28 @@ void UStateTree::ResetCompiled()
 
 void UStateTree::GetAssetRegistryTags(TArray<FAssetRegistryTag>& OutTags) const
 {
-	const FString SchemaClassName = Schema ? Schema->GetClass()->GetName() : TEXT("");
+	const FString SchemaClassName = Schema ? Schema->GetClass()->GetPathName() : TEXT("");
 	OutTags.Add(FAssetRegistryTag(UE::StateTree::SchemaTag, SchemaClassName, FAssetRegistryTag::TT_Alphabetical));
 
 	Super::GetAssetRegistryTags(OutTags);
 }
+
+void UStateTree::PostLoadAssetRegistryTags(const FAssetData& InAssetData, TArray<FAssetRegistryTag>& OutTagsAndValuesToUpdate) const
+{
+	Super::PostLoadAssetRegistryTags(InAssetData, OutTagsAndValuesToUpdate);
+
+	static const FName SchemaTag(TEXT("Schema"));
+	FString SchemaTagValue = InAssetData.GetTagValueRef<FString>(SchemaTag);
+	if (!SchemaTagValue.IsEmpty() && FPackageName::IsShortPackageName(SchemaTagValue))
+	{
+		FTopLevelAssetPath SchemaTagClassPathName = UClass::TryConvertShortTypeNameToPathName<UStruct>(SchemaTagValue, ELogVerbosity::Warning, TEXT("UStateTree::PostLoadAssetRegistryTags"));
+		if (!SchemaTagClassPathName.IsNull())
+		{
+			OutTagsAndValuesToUpdate.Add(FAssetRegistryTag(SchemaTag, SchemaTagClassPathName.ToString(), FAssetRegistryTag::TT_Alphabetical));
+		}
+	}
+}
+
 #endif // WITH_EDITOR
 
 void UStateTree::PostLoad()
diff --git a/Engine/Plugins/Runtime/StateTree/Source/StateTreeModule/Public/StateTree.h b/Engine/Plugins/Runtime/StateTree/Source/StateTreeModule/Public/StateTree.h
--- a/Engine/Plugins/Runtime/StateTree/Source/StateTreeModule/Public/StateTree.h
+++ b/Engine/Plugins/Runtime/StateTree/Source/StateTreeModule/Public/StateTree.h
@@ -127,6 +127,7 @@ protected:
 	
 #if WITH_EDITOR
 	virtual void GetAssetRegistryTags(TArray<FAssetRegistryTag>& OutTags) const override;
+	virtual void PostLoadAssetRegistryTags(const FAssetData& InAssetData, TArray<FAssetRegistryTag>& OutTagsAndValuesToUpdate) const;
 #endif
 
 private:
diff --git a/Engine/Plugins/Runtime/XR/XRVisualization/Source/XRVisualization/Private/XRVisualizationFunctionLibrary.cpp b/Engine/Plugins/Runtime/XR/XRVisualization/Source/XRVisualization/Private/XRVisualizationFunctionLibrary.cpp
--- a/Engine/Plugins/Runtime/XR/XRVisualization/Source/XRVisualization/Private/XRVisualizationFunctionLibrary.cpp
+++ b/Engine/Plugins/Runtime/XR/XRVisualization/Source/XRVisualization/Private/XRVisualizationFunctionLibrary.cpp
@@ -224,6 +224,7 @@ void UXRVisualizationFunctionLibrary::RenderHandMesh(const FXRMotionControllerDa
 					
 					ProceduralMeshComponent->CreateMeshSection_LinearColor(0, Vertices, Indices, Normals, {}, Colors, {}, false);
 
+					// RobM: I couldn't find this section in any of the ini files, does it still exist? if so can HandMeshMaterial be converted to a path name?
 					if (GConfig->DoesSectionExist(TEXT("/Script/EngineSettings.XRVisualizationSettings"), GGameIni))
 					{
 						FString MaterialAddress;
@@ -232,7 +233,7 @@ void UXRVisualizationFunctionLibrary::RenderHandMesh(const FXRMotionControllerDa
 
 						if (!MaterialAddress.IsEmpty())
 						{
-							UMaterialInterface* MaterialInterface = Cast<UMaterialInterface>(StaticFindObject(UMaterialInterface::StaticClass(), ANY_PACKAGE, *MaterialAddress));
+							UMaterialInterface* MaterialInterface = Cast<UMaterialInterface>(StaticFindFirstObject(UMaterialInterface::StaticClass(), *MaterialAddress));
 							ProceduralMeshComponent->SetMaterial(0, MaterialInterface);
 						}
 					}
diff --git a/Engine/Plugins/Runtime/nDisplay/Source/DisplayClusterConfigurator/Private/DisplayClusterConfiguratorModule.cpp b/Engine/Plugins/Runtime/nDisplay/Source/DisplayClusterConfigurator/Private/DisplayClusterConfiguratorModule.cpp
--- a/Engine/Plugins/Runtime/nDisplay/Source/DisplayClusterConfigurator/Private/DisplayClusterConfiguratorModule.cpp
+++ b/Engine/Plugins/Runtime/nDisplay/Source/DisplayClusterConfigurator/Private/DisplayClusterConfiguratorModule.cpp
@@ -321,7 +321,7 @@ void FDisplayClusterConfiguratorModule::RegisterPlacementModeItems()
 			TEXT("/nDisplay/LightCard/LightCard"),
 			TEXT("/nDisplay/LightCard"),
 			TEXT("LightCard"),
-			TEXT("Blueprint")
+			FTopLevelAssetPath(TEXT("/Script/Engine"), TEXT("Blueprint"))
 		);
 		check (LightCardAssetData.IsValid());
 	
