Deserializing UDP message with an older protocol will no longer result in warnings about using short type names.
Old protocols will always use short type names so there's no point in spamming with warnings.

diff --git a/Engine/Plugins/Messaging/UdpMessaging/Source/UdpMessaging/Private/Transport/UdpDeserializedMessage.cpp b/Engine/Plugins/Messaging/UdpMessaging/Source/UdpMessaging/Private/Transport/UdpDeserializedMessage.cpp
--- a/Engine/Plugins/Messaging/UdpMessaging/Source/UdpMessaging/Private/Transport/UdpDeserializedMessage.cpp
+++ b/Engine/Plugins/Messaging/UdpMessaging/Source/UdpMessaging/Private/Transport/UdpDeserializedMessage.cpp
@@ -156,7 +156,7 @@ bool FUdpDeserializedMessageDetails::DeserializeV10(FUdpDeserializedMessage& Des
 		MessageReader << MessageType;
 
 		// @todo gmp: cache message types for faster lookup
-		DeserializedMessage.TypeInfo = UClass::TryFindTypeSlowSafe<UScriptStruct>(MessageType.ToString());
+		DeserializedMessage.TypeInfo = FindFirstObjectSafe<UScriptStruct>(*MessageType.ToString(), EFindFirstObjectOptions::EnsureIfAmbiguous);
 
 		if (!DeserializedMessage.TypeInfo.IsValid(false, true))
 		{
@@ -242,7 +242,7 @@ bool FUdpDeserializedMessageDetails::DeserializeV11_15(FUdpDeserializedMessage&
 		MessageReader << MessageType;
 
 		// @todo gmp: cache message types for faster lookup
-		DeserializedMessage.TypeInfo = UClass::TryFindTypeSlowSafe<UScriptStruct>(MessageType.ToString());
+		DeserializedMessage.TypeInfo = FindFirstObjectSafe<UScriptStruct>(*MessageType.ToString(), EFindFirstObjectOptions::EnsureIfAmbiguous);
 
 		if (!DeserializedMessage.TypeInfo.IsValid(false, true))
 		{
