Add a tracking table to invalidate endpoints that may be unreachable.   If an endpoints becomes unreachable because of a routing issue it is added to the table and exluded from message bus communication.

We also added support for an exclusion list to exlude a specified list of actors. There are CVars added to aid in testing and clearing these lists at runtime.

diff --git a/Engine/Plugins/Messaging/UdpMessaging/Source/UdpMessaging/Private/Tests/UdpMessageTransportTest.cpp b/Engine/Plugins/Messaging/UdpMessaging/Source/UdpMessaging/Private/Tests/UdpMessageTransportTest.cpp
--- a/Engine/Plugins/Messaging/UdpMessaging/Source/UdpMessaging/Private/Tests/UdpMessageTransportTest.cpp
+++ b/Engine/Plugins/Messaging/UdpMessaging/Source/UdpMessaging/Private/Tests/UdpMessageTransportTest.cpp
@@ -23,7 +23,10 @@ public:
 		: NumReceivedMessages(0)
 	{
 		TArray<FIPv4Endpoint> StaticEndpoints;
-		Transport = MakeShared<FUdpMessageTransport, ESPMode::ThreadSafe>(UnicastEndpoint, MulticastEndpoint, MoveTemp(StaticEndpoints), MulticastTimeToLive);
+		TArray<FIPv4Endpoint> ExcludeEndpoints;
+		Transport = MakeShared<FUdpMessageTransport, ESPMode::ThreadSafe>(UnicastEndpoint, MulticastEndpoint,
+																		  MoveTemp(StaticEndpoints),
+																		  MoveTemp(ExcludeEndpoints), MulticastTimeToLive);
 	}
 
 public:
diff --git a/Engine/Plugins/Messaging/UdpMessaging/Source/UdpMessaging/Private/Transport/UdpMessageProcessor.cpp b/Engine/Plugins/Messaging/UdpMessaging/Source/UdpMessaging/Private/Transport/UdpMessageProcessor.cpp
--- a/Engine/Plugins/Messaging/UdpMessaging/Source/UdpMessaging/Private/Transport/UdpMessageProcessor.cpp
+++ b/Engine/Plugins/Messaging/UdpMessaging/Source/UdpMessaging/Private/Transport/UdpMessageProcessor.cpp
@@ -3,7 +3,9 @@
 #include "Transport/UdpMessageProcessor.h"
 #include "Algo/AllOf.h"
 #include "Algo/Transform.h"
+#include "HAL/IConsoleManager.h"
 #include "INetworkMessagingExtension.h"
+#include "Interfaces/IPv4/IPv4Endpoint.h"
 #include "UdpMessagingPrivate.h"
 
 #include "Common/UdpSocketSender.h"
@@ -30,9 +32,51 @@
 
 const int32 FUdpMessageProcessor::DeadHelloIntervals = 5;
 
+TAutoConsoleVariable<int32> CVarFakeSocketError(
+	TEXT("MessageBus.UDP.InduceSocketError"),
+	0,
+	TEXT("This CVar can be used to induce a socket failure on outbound communication.\n")
+	TEXT("Any non zero value will force the output socket connection to fail if the IP address matches\n")
+	TEXT("one of the values in MessageBus.UDP.ConnectionsToError. The list can be cleared by invoking\n")
+	TEXT("MessageBus.UDP.ClearDenyList."),
+	ECVF_Default
+);
+
+TAutoConsoleVariable<FString> CVarConnectionsToError(
+	TEXT("MessageBus.UDP.ConnectionsToError"),
+	TEXT(""),
+	TEXT("Connections to error out on when MessageBus.UDP.InduceSocketError is enabled.\n")
+	TEXT("This can be a comma separated list in the form IPAddr2:port,IPAddr3:port"),
+	ECVF_Default);
 
 namespace UE::Private::MessageProcessor
 {
+
+bool ShouldErrorOnConnection(const FIPv4Endpoint& InEndpoint)
+{
+	if (CVarFakeSocketError.GetValueOnAnyThread() == 0)
+	{
+		return false;
+	}
+
+	const FString ConnectionsToErrorString = CVarConnectionsToError.GetValueOnAnyThread();
+	TArray<FString> EndpointStrings;
+	ConnectionsToErrorString.ParseIntoArray(EndpointStrings, TEXT(","));
+	for (const FString& EndpointString : EndpointStrings)
+	{
+		FIPv4Endpoint Endpoint;
+		if (FIPv4Endpoint::Parse(EndpointString, Endpoint))
+		{
+			if (Endpoint.Address == InEndpoint.Address &&
+				(Endpoint.Port == 0 || Endpoint.Port == InEndpoint.Port))
+			{
+				return true;
+			}
+		}
+	}
+	return false;
+}
+
 FOnTransferDataUpdated& OnSegmenterUpdated()
 {
 	static FOnTransferDataUpdated OnTransferUpdated;
@@ -405,7 +449,7 @@ void FUdpMessageProcessor::ConsumeInboundSegments()
 		FUdpMessageSegment::FHeader Header;
 		*Segment.Data << Header;
 
-		if (FilterSegment(Header))
+		if (FilterSegment(Header, Segment.Sender))
 		{
 			FNodeInfo& NodeInfo = KnownNodes.FindOrAdd(Header.SenderNodeId);
 
@@ -563,7 +607,7 @@ void FUdpMessageProcessor::ConsumeOutboundMessages()
 	}
 }
 
-bool FUdpMessageProcessor::FilterSegment(const FUdpMessageSegment::FHeader& Header)
+bool FUdpMessageProcessor::FilterSegment(const FUdpMessageSegment::FHeader& Header, const FIPv4Endpoint& Sender)
 {
 	// filter locally generated segments
 	if (Header.SenderNodeId == LocalNodeId)
@@ -571,6 +615,11 @@ bool FUdpMessageProcessor::FilterSegment(const FUdpMessageSegment::FHeader& Head
 		return false;
 	}
 
+	if (!CanAcceptEndpointDelegate.Execute(Header.SenderNodeId, Sender))
+	{
+		return false;
+	}
+
 	// filter unsupported protocol versions
 	if (!SupportedProtocolVersions.Contains(Header.ProtocolVersion))
 	{
@@ -864,6 +913,13 @@ bool FUdpMessageProcessor::MoreToSend()
 	return false;
 }
 
+void FUdpMessageProcessor::HandleSocketError(const FNodeInfo& NodeInfo) const
+{
+	UE_LOG(LogUdpMessaging, Error,
+		   TEXT("Socket error detected when communicating with %s. Banning communication to that endpoint."), *NodeInfo.Endpoint.ToString());
+	ErrorSendingToEndpointDelegate.Execute(NodeInfo.NodeId, NodeInfo.Endpoint);
+}
+
 void FUdpMessageProcessor::UpdateKnownNodes()
 {
 	SCOPED_MESSAGING_TRACE(FUdpMessageProcessor_UpdateKnownNodes);
@@ -879,9 +935,10 @@ void FUdpMessageProcessor::UpdateKnownNodes()
 		int32 NodeByteSent = UpdateSegmenters(KnownNodePair.Value);
 		// if NodByteSent is negative, there is a socket error, continuing is useless
 		bSuccess = NodeByteSent >= 0;
-		if (!bSuccess)
+		if (!bSuccess || UE::Private::MessageProcessor::ShouldErrorOnConnection(KnownNodePair.Value.Endpoint))
 		{
-			UE_LOG(LogUdpMessaging, Verbose, TEXT("FUdpMessageProcessor::UpdateKnownNodes received negative NodeByteSent (%d) from socket sender."), NodeByteSent);
+			bSuccess = false; // To ensure we trigger the delegate on a forced error.
+			HandleSocketError(KnownNodePair.Value);
 			break;
 		}
 
@@ -889,6 +946,7 @@ void FUdpMessageProcessor::UpdateKnownNodes()
 		// if there is a socket error, continuing is useless
 		if (!bSuccess)
 		{
+			HandleSocketError(KnownNodePair.Value);
 			break;
 		}
 
diff --git a/Engine/Plugins/Messaging/UdpMessaging/Source/UdpMessaging/Private/Transport/UdpMessageProcessor.h b/Engine/Plugins/Messaging/UdpMessaging/Source/UdpMessaging/Private/Transport/UdpMessageProcessor.h
--- a/Engine/Plugins/Messaging/UdpMessaging/Source/UdpMessaging/Private/Transport/UdpMessageProcessor.h
+++ b/Engine/Plugins/Messaging/UdpMessaging/Source/UdpMessaging/Private/Transport/UdpMessageProcessor.h
@@ -586,6 +586,30 @@
 		return ErrorDelegate;
 	}
 
+	/** 
+	 * Returns a delegate that is executed when a socket fails to communicate
+	 * upon sending to a target endpoint.
+	 * @return The delegate
+	 * @note this delegate is broadcasted from the processor thread.
+	 */
+	DECLARE_DELEGATE_TwoParams(FOnErrorSendingToEndpoint, const FGuid& /*NodeId*/, const FIPv4Endpoint& /*SendersIpAddress*/)
+	FOnErrorSendingToEndpoint& OnErrorSendingToEndpoint_UdpMessageProcessorThread()
+	{
+		return ErrorSendingToEndpointDelegate;
+	}
+
+	/**
+	 * Returns a delegate that is executed when a socket fails to communicate
+	 * upon sending to a target endpoint.
+	 * @return The delegate
+	 * @note this delegate is broadcasted from the processor thread.
+	 */
+	DECLARE_DELEGATE_RetVal_TwoParams(bool, FCanAcceptEndpoint, const FGuid& /*NodeId*/, const FIPv4Endpoint& /*SendersIpAddress*/)
+	FCanAcceptEndpoint& OnCanAcceptEndpoint_UdpMessageProcessorThread()
+	{
+		return CanAcceptEndpointDelegate;
+	}
+
 	TArray<FIPv4Endpoint> GetKnownEndpoints() const;
 
 public:
@@ -630,7 +654,7 @@
 	 * @param Sender The segment sender.
 	 * @return true if the segment passed the filter, false otherwise.
 	 */
-	bool FilterSegment(const FUdpMessageSegment::FHeader& Header);
+	bool FilterSegment(const FUdpMessageSegment::FHeader& Header, const FIPv4Endpoint& Sender);
 
 	/**
 	 * Processes an Abort segment.
@@ -767,6 +791,9 @@
 
 private:
 
+	/** Handles a communication error to a particular endpoint.*/
+	void HandleSocketError(const FNodeInfo& NodeInfo) const;
+
 	/** Checks all known nodes to see if any segmenters have NeedSending set to true. */
 	bool MoreToSend();
 
@@ -845,6 +872,12 @@
 	/** Holds a delegate to be invoked when a socket error happen. */
 	FOnError ErrorDelegate;
 
+	/** Holds a delegate to be invoked when a socket error occurs sending to a given endpoint. */
+	FOnErrorSendingToEndpoint ErrorSendingToEndpointDelegate;
+
+	/** Holds a delegate to be invoked when checking the validity of a given endpoint address. */
+	FCanAcceptEndpoint CanAcceptEndpointDelegate;
+
 	/** The configured message format (from UUdpMessagingSettings). */
 	EUdpMessageFormat MessageFormat;
 
diff --git a/Engine/Plugins/Messaging/UdpMessaging/Source/UdpMessaging/Private/Transport/UdpMessageTransport.cpp b/Engine/Plugins/Messaging/UdpMessaging/Source/UdpMessaging/Private/Transport/UdpMessageTransport.cpp
--- a/Engine/Plugins/Messaging/UdpMessaging/Source/UdpMessaging/Private/Transport/UdpMessageTransport.cpp
+++ b/Engine/Plugins/Messaging/UdpMessaging/Source/UdpMessaging/Private/Transport/UdpMessageTransport.cpp
@@ -1,8 +1,10 @@
 // Copyright Epic Games, Inc. All Rights Reserved.
 
 #include "Transport/UdpMessageTransport.h"
+#include "Algo/AnyOf.h"
 #include "INetworkMessagingExtension.h"
 #include "Interfaces/IPv4/IPv4Address.h"
+#include "Interfaces/IPv4/IPv4Endpoint.h"
 #include "UdpMessagingPrivate.h"
 
 #include "Common/UdpSocketBuilder.h"
@@ -29,7 +31,22 @@
 /* FUdpMessageTransport structors
  *****************************************************************************/
 
-FUdpMessageTransport::FUdpMessageTransport(const FIPv4Endpoint& InUnicastEndpoint, const FIPv4Endpoint& InMulticastEndpoint, TArray<FIPv4Endpoint> InStaticEndpoints, uint8 InMulticastTtl)
+TAutoConsoleVariable<int32> CVarMaxRetriesForBadEndpoint(
+	TEXT("MessageBus.UDP.MaxRetriesForBadEndpoint"),
+	5,
+	TEXT("The maximum number of retries that will be attempted when a socket connection fails to reach an endpoint."),
+	ECVF_Default
+);
+
+TAutoConsoleVariable<bool> CVarEndpointDenyListEnabled(
+	TEXT("MessageBus.UDP.EndpointDenyListEnabled"),
+	true,
+	TEXT("Specifies if the endpoint deny list is enabled. If enabled, a problematic endpoint will be tagged for possible exclusion from communication.")
+	TEXT("The maximum attempts allowed is determined by MessageBus.UDP.MaxRetriesForBadEndpoint"),
+	ECVF_Default
+);
+
+FUdpMessageTransport::FUdpMessageTransport(const FIPv4Endpoint& InUnicastEndpoint, const FIPv4Endpoint& InMulticastEndpoint, TArray<FIPv4Endpoint> InStaticEndpoints, TArray<FIPv4Endpoint> InExcludedEndpoints, uint8 InMulticastTtl)
 	: MessageProcessor(nullptr)
 	, MulticastEndpoint(InMulticastEndpoint)
 	, MulticastReceiver(nullptr)
@@ -43,6 +60,9 @@ FUdpMessageTransport::FUdpMessageTransport(const FIPv4Endpoint& InUnicastEndpoin
 	, UnicastSocket(nullptr)
 #endif
 	, StaticEndpoints(MoveTemp(InStaticEndpoints))
+	, ExcludedEndpoints(MoveTemp(InExcludedEndpoints))
+	, ClearDenyList(TEXT("MessageBus.UDP.ClearDenyList"), TEXT("Clear the UDP socket deny list."),
+					FConsoleCommandDelegate::CreateRaw(this, &FUdpMessageTransport::DoClearDenyCandidateList))
 {
 }
 
@@ -242,6 +262,8 @@ bool FUdpMessageTransport::StartTransport(IMessageTransportHandler& Handler)
 		MessageProcessor->OnNodeDiscovered().BindRaw(this, &FUdpMessageTransport::HandleProcessorNodeDiscovered);
 		MessageProcessor->OnNodeLost().BindRaw(this, &FUdpMessageTransport::HandleProcessorNodeLost);
 		MessageProcessor->OnError().BindRaw(this, &FUdpMessageTransport::HandleProcessorError);
+		MessageProcessor->OnErrorSendingToEndpoint_UdpMessageProcessorThread().BindRaw(this, &FUdpMessageTransport::HandleEndpointCommunicationError);
+		MessageProcessor->OnCanAcceptEndpoint_UdpMessageProcessorThread().BindRaw(this, &FUdpMessageTransport::HandleProcessorEndpointCheck);
 	}
 
 	if (MulticastSocket != nullptr)
@@ -355,6 +377,38 @@ void FUdpMessageTransport::HandleProcessorNodeLost(const FGuid& LostNodeId)
 	TransportHandler->ForgetTransportNode(LostNodeId);
 }
 
+bool FUdpMessageTransport::HandleProcessorEndpointCheck(const FGuid& EndpointNodeId, const FIPv4Endpoint& SenderIpAddress)
+{
+	if (CVarEndpointDenyListEnabled.GetValueOnAnyThread())
+	{
+		FDenyCandidate* DenyCandidate = DenyCandidateList.Find(EndpointNodeId);
+		if (DenyCandidate && DenyCandidate->EndpointFailureCount > CVarMaxRetriesForBadEndpoint.GetValueOnAnyThread())
+		{
+			return false;
+		}
+	}
+
+	for (const FIPv4Endpoint& Endpoint : ExcludedEndpoints)
+	{
+		if ((Endpoint.Port == 0 && SenderIpAddress.Address == Endpoint.Address)
+			|| SenderIpAddress == Endpoint)
+		{
+			return false;
+		}
+	}
+	return true;
+}
+
+void FUdpMessageTransport::DoClearDenyCandidateList()
+{
+	DenyCandidateList.Reset();
+}
+
+void FUdpMessageTransport::HandleEndpointCommunicationError(const FGuid& EndpointId, const FIPv4Endpoint& /*Unused EndpointIdAddress*/)
+{
+	FDenyCandidate& DenyCandidate = DenyCandidateList.FindOrAdd(EndpointId);
+	DenyCandidate.EndpointFailureCount++;
+}
 
 void FUdpMessageTransport::HandleProcessorError()
 {
diff --git a/Engine/Plugins/Messaging/UdpMessaging/Source/UdpMessaging/Private/Transport/UdpMessageTransport.h b/Engine/Plugins/Messaging/UdpMessaging/Source/UdpMessaging/Private/Transport/UdpMessageTransport.h
--- a/Engine/Plugins/Messaging/UdpMessaging/Source/UdpMessaging/Private/Transport/UdpMessageTransport.h
+++ b/Engine/Plugins/Messaging/UdpMessaging/Source/UdpMessaging/Private/Transport/UdpMessageTransport.h
@@ -11,6 +11,7 @@
 #include "Templates/SharedPointer.h"
 #include "Async/Future.h"
 #include "Containers/Ticker.h"
+#include "HAL/IConsoleManager.h"
 
 class FArrayReader;
 class FUdpReassembledMessage;
@@ -44,9 +45,15 @@ public:
 	 *
 	 * @param InLocalEndpoint The local IP endpoint to receive messages on.
 	 * @param InMulticastEndpoint The multicast group endpoint to transport messages to.
+	 * @param InStaticEndpoints array of endpoints to include in communication
+	 * @param InExcludedEndpoints array of endpoints to exclude in communication
 	 * @param InMulticastTtl The multicast time-to-live.
 	 */
-	FUdpMessageTransport(const FIPv4Endpoint& InLocalEndpoint, const FIPv4Endpoint& InMulticastEndpoint, TArray<FIPv4Endpoint> InStaticEndpoints, uint8 InMulticastTtl);
+	FUdpMessageTransport(const FIPv4Endpoint& InLocalEndpoint,
+						 const FIPv4Endpoint& InMulticastEndpoint,
+						 TArray<FIPv4Endpoint> InStaticEndpoints,
+						 TArray<FIPv4Endpoint> InExcludedEndpoints,
+						 uint8 InMulticastTtl);
 
 	/** Virtual destructor. */
 	virtual ~FUdpMessageTransport();
@@ -103,6 +110,12 @@ private:
 	/** */
 	void HandleProcessorError();
 
+	/** Returns true if we can add this endpoint to the list of nodes we communicate with. */
+	bool HandleProcessorEndpointCheck(const FGuid& EndpointNodeId, const FIPv4Endpoint& SenderIpAddress);
+
+	/** Handle endpoint communication error. */
+	void HandleEndpointCommunicationError(const FGuid& EndpointId, const FIPv4Endpoint& EndpointIdAddress);
+
 	/** Launches a routine that attempts to restart the transport. */
 	void StartAutoRepairRoutine(uint32 MaxRetryAttempt);
 
@@ -119,6 +132,9 @@ private:
 	void HandleSocketDataReceived(const TSharedPtr<FArrayReader, ESPMode::ThreadSafe>& Data, const FIPv4Endpoint& Sender);
 
 private:
+	/** Clear the deny list */
+	void DoClearDenyCandidateList();
+
 	/** Holds any pending restart request. */
 	TFuture<void> ErrorFuture;
 
@@ -157,6 +173,20 @@ private:
 	FSocket* UnicastSocket;
 #endif
 
+	struct FDenyCandidate
+	{
+		int32 EndpointFailureCount = 0;
+	};
+
+	/** Deny list of node ids not allowed to talk to UDP processor */
+	TMap<FGuid, FDenyCandidate> DenyCandidateList;
+
 	/** Holds the static endpoints. */
 	TSet<FIPv4Endpoint> StaticEndpoints;
+
+	/** Excluded list of IP addresses not allowed to talk to UDP processor */
+	TArray<FIPv4Endpoint> ExcludedEndpoints;
+
+	/** Console command to reset the deny list */
+	FAutoConsoleCommand ClearDenyList;
 };
diff --git a/Engine/Plugins/Messaging/UdpMessaging/Source/UdpMessaging/Private/UdpMessagingModule.cpp b/Engine/Plugins/Messaging/UdpMessaging/Source/UdpMessaging/Private/UdpMessagingModule.cpp
--- a/Engine/Plugins/Messaging/UdpMessaging/Source/UdpMessaging/Private/UdpMessagingModule.cpp
+++ b/Engine/Plugins/Messaging/UdpMessaging/Source/UdpMessaging/Private/UdpMessagingModule.cpp
@@ -472,6 +472,25 @@ protected:
 		return bParsedAddr;
 	}
 
+	TArray<FIPv4Endpoint> ParseStringArrayAddresses(const TArray<FString>& StringAddresses)
+	{
+		TArray<FIPv4Endpoint> Endpoints;
+		for (const FString& EndpointAsString : StringAddresses)
+		{
+			FIPv4Endpoint Endpoint;
+
+			if (ParseEndpoint(EndpointAsString, Endpoint))
+			{
+				Endpoints.Add(Endpoint);
+			}
+			else
+			{
+				UE_LOG(LogUdpMessaging, Warning, TEXT("Invalid UDP Messaging Endpoint '%s'"), *EndpointAsString);
+			}
+		}
+		return Endpoints;
+	}
+
 	/** Initializes the message bridge with the current settings. */
 	void InitializeBridge()
 	{
@@ -509,20 +528,10 @@ protected:
 
 		// Initialize the service with the additional endpoints added through the modular interface
 		TArray<FIPv4Endpoint> StaticEndpoints = AdditionalStaticEndpoints.Array();
+		StaticEndpoints += ParseStringArrayAddresses(Settings->StaticEndpoints);
 
-		for (auto& StaticEndpoint : Settings->StaticEndpoints)
-		{
-			FIPv4Endpoint Endpoint;
-
-			if (ParseEndpoint(StaticEndpoint, Endpoint))
-			{
-				StaticEndpoints.Add(Endpoint);
-			}
-			else
-			{
-				UE_LOG(LogUdpMessaging, Warning, TEXT("Invalid UDP Messaging Static Endpoint '%s'"), *StaticEndpoint);
-			}
-		}
+		// Addresses to deny on transport.
+		TArray<FIPv4Endpoint> ExcludedEndpoints = ParseStringArrayAddresses(Settings->ExcludedEndpoints);
 
 		if (Settings->MulticastTimeToLive == 0)
 		{
@@ -536,7 +545,8 @@ protected:
 		}
 		UE_LOG(LogUdpMessaging, Log, TEXT("Initializing bridge on interface %s to multicast group %s."), *UnicastEndpoint.ToString(), *MulticastEndpoint.ToText().ToString());
 
-		TSharedRef<FUdpMessageTransport, ESPMode::ThreadSafe> Transport = MakeShared<FUdpMessageTransport, ESPMode::ThreadSafe>(UnicastEndpoint, MulticastEndpoint, MoveTemp(StaticEndpoints), Settings->MulticastTimeToLive);
+		TSharedRef<FUdpMessageTransport, ESPMode::ThreadSafe> Transport = MakeShared<FUdpMessageTransport, ESPMode::ThreadSafe>(
+			UnicastEndpoint, MulticastEndpoint, MoveTemp(StaticEndpoints), MoveTemp(ExcludedEndpoints), Settings->MulticastTimeToLive);
 		WeakBridgeTransport = Transport;
 		MessageBridge = FMessageBridgeBuilder()
 			.UsingTransport(Transport);
diff --git a/Engine/Plugins/Messaging/UdpMessaging/Source/UdpMessaging/Public/Shared/UdpMessagingSettings.h b/Engine/Plugins/Messaging/UdpMessaging/Source/UdpMessaging/Public/Shared/UdpMessagingSettings.h
--- a/Engine/Plugins/Messaging/UdpMessaging/Source/UdpMessaging/Public/Shared/UdpMessagingSettings.h
+++ b/Engine/Plugins/Messaging/UdpMessaging/Source/UdpMessaging/Public/Shared/UdpMessagingSettings.h
@@ -128,6 +128,14 @@ public:
 	UPROPERTY(config, EditAnywhere, Category=Transport, AdvancedDisplay)
 	TArray<FString> StaticEndpoints;
 
+	/**
+	 * List of IP addresses that are banned from communicating with this client.
+	 *
+	 * The format is IP_ADDRESS:PORT_NUMBER.  If port number is 0 then all ports are blocked.
+	 */
+	UPROPERTY(config, EditAnywhere, Category=Transport, AdvancedDisplay)
+	TArray<FString> ExcludedEndpoints;
+
 public:
 
 	/** Whether the UDP tunnel is enabled. */
