Make sure we stop message bus on error.

diff --git a/Engine/Plugins/Messaging/UdpMessaging/Source/UdpMessaging/Private/Transport/UdpMessageProcessor.cpp b/Engine/Plugins/Messaging/UdpMessaging/Source/UdpMessaging/Private/Transport/UdpMessageProcessor.cpp
--- a/Engine/Plugins/Messaging/UdpMessaging/Source/UdpMessaging/Private/Transport/UdpMessageProcessor.cpp
+++ b/Engine/Plugins/Messaging/UdpMessaging/Source/UdpMessaging/Private/Transport/UdpMessageProcessor.cpp
@@ -956,6 +956,7 @@ void FUdpMessageProcessor::UpdateKnownNodes()
 	// if we had socket error, fire up the error delegate
 	if (!bSuccess || Beacon->HasSocketError())
 	{
+		bStopping = true;
 		ErrorDelegate.ExecuteIfBound();
 	}
 }
