Add a bad endpoint period per post review comments.

diff --git a/Engine/Plugins/Messaging/UdpMessaging/Source/UdpMessaging/Private/Transport/UdpMessageTransport.cpp b/Engine/Plugins/Messaging/UdpMessaging/Source/UdpMessaging/Private/Transport/UdpMessageTransport.cpp
--- a/Engine/Plugins/Messaging/UdpMessaging/Source/UdpMessaging/Private/Transport/UdpMessageTransport.cpp
+++ b/Engine/Plugins/Messaging/UdpMessaging/Source/UdpMessaging/Private/Transport/UdpMessageTransport.cpp
@@ -38,6 +38,13 @@ TAutoConsoleVariable<int32> CVarMaxRetriesForBadEndpoint(
 	ECVF_Default
 );
 
+TAutoConsoleVariable<int32> CVarBadEndpointPeriod(
+	TEXT("MessageBus.UDP.BadEndpointPeriod"),
+	60,
+	TEXT("The period of time, in seconds, between endpoint socket errors to be considered a bad endpoint."),
+	ECVF_Default
+);
+
 TAutoConsoleVariable<bool> CVarEndpointDenyListEnabled(
 	TEXT("MessageBus.UDP.EndpointDenyListEnabled"),
 	true,
@@ -407,7 +414,16 @@ void FUdpMessageTransport::DoClearDenyCandidateList()
 void FUdpMessageTransport::HandleEndpointCommunicationError(const FGuid& EndpointId, const FIPv4Endpoint& /*Unused EndpointIdAddress*/)
 {
 	FDenyCandidate& DenyCandidate = DenyCandidateList.FindOrAdd(EndpointId);
-	DenyCandidate.EndpointFailureCount++;
+	FDateTime CurrentTime = FDateTime::UtcNow();
+	if ((CurrentTime - DenyCandidate.LastFailTime).GetSeconds() < CVarBadEndpointPeriod.GetValueOnAnyThread())
+	{
+		DenyCandidate.EndpointFailureCount++;
+	}
+	else
+	{
+		DenyCandidate.EndpointFailureCount = 0;
+	}
+	DenyCandidate.LastFailTime = CurrentTime;
 }
 
 void FUdpMessageTransport::HandleProcessorError()
diff --git a/Engine/Plugins/Messaging/UdpMessaging/Source/UdpMessaging/Private/Transport/UdpMessageTransport.h b/Engine/Plugins/Messaging/UdpMessaging/Source/UdpMessaging/Private/Transport/UdpMessageTransport.h
--- a/Engine/Plugins/Messaging/UdpMessaging/Source/UdpMessaging/Private/Transport/UdpMessageTransport.h
+++ b/Engine/Plugins/Messaging/UdpMessaging/Source/UdpMessaging/Private/Transport/UdpMessageTransport.h
@@ -176,6 +176,7 @@ private:
 	struct FDenyCandidate
 	{
 		int32 EndpointFailureCount = 0;
+		FDateTime LastFailTime = FDateTime::UtcNow();
 	};
 
 	/** Deny list of node ids not allowed to talk to UDP processor */
