Add a settings to control the size of the work send queue. Sessions with a large number of transactions and clients that join late may cause this queue to fill up.  This is a way for users to work around the issue on the server side. Additional changes are required on the UDP messaging side to mitigate this issue in the future.

diff --git a/Engine/Plugins/Messaging/UdpMessaging/Source/UdpMessaging/Private/Transport/UdpMessageProcessor.cpp b/Engine/Plugins/Messaging/UdpMessaging/Source/UdpMessaging/Private/Transport/UdpMessageProcessor.cpp
--- a/Engine/Plugins/Messaging/UdpMessaging/Source/UdpMessaging/Private/Transport/UdpMessageProcessor.cpp
+++ b/Engine/Plugins/Messaging/UdpMessaging/Source/UdpMessaging/Private/Transport/UdpMessageProcessor.cpp
@@ -6,6 +6,7 @@
 #include "HAL/IConsoleManager.h"
 #include "INetworkMessagingExtension.h"
 #include "Interfaces/IPv4/IPv4Endpoint.h"
+#include "Misc/AssertionMacros.h"
 #include "UdpMessagingPrivate.h"
 
 #include "Common/UdpSocketSender.h"
@@ -89,9 +90,23 @@ FOnInboundTransferDataUpdated& OnReassemblerUpdated()
 	return OnTransferUpdated;
 }
 
+uint16 GetMessageProcessorWorkQueueSize()
+{
+	return GetDefault<UUdpMessagingSettings>()->WorkQueueSize;
+}
+
 }
 /* FUdpMessageProcessor structors
  *****************************************************************************/
+FUdpMessageProcessor::FNodeInfo::FNodeInfo()
+	: LastSegmentReceivedTime(FDateTime::MinValue())
+	, NodeId()
+	, ProtocolVersion(UDP_MESSAGING_TRANSPORT_PROTOCOL_VERSION)
+	, WorkQueue(UE::Private::MessageProcessor::GetMessageProcessorWorkQueueSize())
+{
+	ComputeWindowSize(0,0);
+}
+
 
 FUdpMessageProcessor::FUdpMessageProcessor(FSocket& InSocket, const FGuid& InNodeId, const FIPv4Endpoint& InMulticastEndpoint)
 	: Beacon(nullptr)
@@ -541,7 +556,17 @@ bool FUdpMessageProcessor::ConsumeOneOutboundMessage(const FOutboundMessage& Out
 					  [FindNodeWithLog](const FGuid &Id) {return FindNodeWithLog(Id);},
 					  [this](const FGuid &Id) {return KnownNodes.Find(Id);} );
 
-	const bool bCanConsume = Algo::AllOf(Recipients, [](FNodeInfo *Node) {return !Node->WorkQueue.IsFull();});
+	const bool bCanConsume = Algo::AllOf(Recipients, [](FNodeInfo *Node)
+		{
+			const bool bWorkQueueIsNotFull = !Node->WorkQueue.IsFull();
+			if (!bWorkQueueIsNotFull)
+			{
+				UE_LOG(LogUdpMessaging, Warning,
+					   TEXT("Work queue for node %s is full. Send queue is congested. Will retry on next tick."), *Node->Endpoint.ToString());
+			}
+			return bWorkQueueIsNotFull;
+		});
+
 	const bool bIsReliable = EnumHasAnyFlags(OutboundMessage.MessageFlags, EMessageFlags::Reliable);
 	if (bCanConsume)
 	{
@@ -582,13 +607,6 @@ void FUdpMessageProcessor::ConsumeOutboundMessages()
 	SCOPED_MESSAGING_TRACE(FUdpMessageProcessor_ConsumeOutputMessages);
 	FOutboundMessage OutboundMessage;
 
-	const auto CannotConsumeMsg = []()
-	{
-		UE_LOG(LogUdpMessaging, Warning, TEXT("Can't consume outbound messages. Work queue on one or more receipents is full. Will try again next tick."));
-		// We can't queue any messages at the moment because our some/all of our work queues are full.
-		//
-	};
-
 	if (DeferredOutboundMessage)
 	{
 		if (ConsumeOneOutboundMessage(DeferredOutboundMessage.GetValue()))
@@ -597,7 +615,6 @@ void FUdpMessageProcessor::ConsumeOutboundMessages()
 		}
 		else
 		{
-			CannotConsumeMsg();
 			return ;
 		}
 	}
@@ -606,7 +623,6 @@ void FUdpMessageProcessor::ConsumeOutboundMessages()
 	{
 		if (!ConsumeOneOutboundMessage(OutboundMessage))
 		{
-			CannotConsumeMsg();
 			DeferredOutboundMessage = OutboundMessage;
 			return ;
 		}
diff --git a/Engine/Plugins/Messaging/UdpMessaging/Source/UdpMessaging/Private/Transport/UdpMessageProcessor.h b/Engine/Plugins/Messaging/UdpMessaging/Source/UdpMessaging/Private/Transport/UdpMessageProcessor.h
--- a/Engine/Plugins/Messaging/UdpMessaging/Source/UdpMessaging/Private/Transport/UdpMessageProcessor.h
+++ b/Engine/Plugins/Messaging/UdpMessaging/Source/UdpMessaging/Private/Transport/UdpMessageProcessor.h
@@ -38,8 +38,6 @@ class FUdpSocketSender;
 class IMessageAttachment;
 enum class EUdpMessageFormat : uint8;
 
-constexpr const uint16 MessageProcessorWorkQueueSize{1024};
-
 /**
  * Running statistics as known by the UdpMessageProcessor. This will be per endpoint.
  */
@@ -118,7 +116,8 @@ FOnOutboundTransferDataUpdated &OnSegmenterUpdated();
 /*
  * Global delegate for handling reassembler (aka received data) updates.
  */
-FOnInboundTransferDataUpdated& OnReassemblerUpdated();
+FOnInboundTransferDataUpdated &OnReassemblerUpdated();
+
 }
 
 /**
@@ -150,7 +149,7 @@ class FUdpMessageProcessor
 		TMap<int32, TSharedPtr<FUdpMessageSegmenter>> Segmenters;
 
 		/** Holds of queue of MessageIds to send. They are processed in round-robin fashion. */
-		TUdpCircularQueue<int32> WorkQueue{MessageProcessorWorkQueueSize};
+		TUdpCircularQueue<int32> WorkQueue;
 
 		/** A map from sequence id to information about what was sent. The size of this map is the size of our sending window. */
 		TMap<uint64, FSentData> InflightSegments;
@@ -168,13 +167,7 @@ class FUdpMessageProcessor
 		FMessageTransportStatistics Statistics;
 
 		/** Default constructor. */
-		FNodeInfo()
-			: LastSegmentReceivedTime(FDateTime::MinValue())
-			, NodeId()
-			, ProtocolVersion(UDP_MESSAGING_TRANSPORT_PROTOCOL_VERSION)
-		{
-			ComputeWindowSize(0,0);
-		}
+		FNodeInfo();
 
 		/**
 		 * Remap a MessageId and SegmentId pair into a 64 bit number using Szudzik pairing calculation.
diff --git a/Engine/Plugins/Messaging/UdpMessaging/Source/UdpMessaging/Private/UdpMessagingModule.cpp b/Engine/Plugins/Messaging/UdpMessaging/Source/UdpMessaging/Private/UdpMessagingModule.cpp
--- a/Engine/Plugins/Messaging/UdpMessaging/Source/UdpMessaging/Private/UdpMessagingModule.cpp
+++ b/Engine/Plugins/Messaging/UdpMessaging/Source/UdpMessaging/Private/UdpMessagingModule.cpp
@@ -630,6 +630,7 @@ protected:
 			FParse::Bool(CommandLine, TEXT("-UDPMESSAGING_TRANSPORT_ENABLE="), Settings->EnableTransport);
 			FParse::Value(CommandLine, TEXT("-UDPMESSAGING_TRANSPORT_UNICAST="), Settings->UnicastEndpoint);
 			FParse::Value(CommandLine, TEXT("-UDPMESSAGING_TRANSPORT_MULTICAST="), Settings->MulticastEndpoint);
+			FParse::Value(CommandLine, TEXT("-UDPMESSAGING_WORK_QUEUE_SIZE="), Settings->WorkQueueSize);
 
 			FString StaticEndpoints;
 			FParse::Value(CommandLine, TEXT("-UDPMESSAGING_TRANSPORT_STATIC="), StaticEndpoints, false);
diff --git a/Engine/Plugins/Messaging/UdpMessaging/Source/UdpMessaging/Public/Shared/UdpMessagingSettings.h b/Engine/Plugins/Messaging/UdpMessaging/Source/UdpMessaging/Public/Shared/UdpMessagingSettings.h
--- a/Engine/Plugins/Messaging/UdpMessaging/Source/UdpMessaging/Public/Shared/UdpMessagingSettings.h
+++ b/Engine/Plugins/Messaging/UdpMessaging/Source/UdpMessaging/Public/Shared/UdpMessagingSettings.h
@@ -84,6 +84,16 @@ public:
 	UPROPERTY(config, EditAnywhere, Category = Transport, AdvancedDisplay)
 	uint32 AutoRepairAttemptLimit = 10;
 
+	/**
+	 * The buffer size for the working send queues. Each node connections
+	 * gets a send queue and new messages to send are put on that queue.  The send
+	 * scheduler will fairly send data on this queue and re-queue when partial data is
+	 * sent or waiting for a reliable message. Can be specified on the command line with
+	 * `-UDPMESSAGING_WORK_QUEUE_SIZE=`
+	 */
+	UPROPERTY(config, EditAnywhere, Category = Transport, meta = (ClampMin="1024"), AdvancedDisplay)
+	uint16 WorkQueueSize = 1024;
+
 	/** Whether to stop the transport service when the application deactivates, and restart it when the application is reactivated */
 	UPROPERTY(config)
 	bool bStopServiceWhenAppDeactivates = true;
