Make UdpMessageProcessor less spammy when in a congested scenario.  There is still an outstanding task to improve the handling of the work queue.

#rnx

diff --git a/Engine/Plugins/Messaging/UdpMessaging/Source/UdpMessaging/Private/Transport/UdpMessageProcessor.cpp b/Engine/Plugins/Messaging/UdpMessaging/Source/UdpMessaging/Private/Transport/UdpMessageProcessor.cpp
--- a/Engine/Plugins/Messaging/UdpMessaging/Source/UdpMessaging/Private/Transport/UdpMessageProcessor.cpp
+++ b/Engine/Plugins/Messaging/UdpMessaging/Source/UdpMessaging/Private/Transport/UdpMessageProcessor.cpp
@@ -556,15 +556,9 @@ bool FUdpMessageProcessor::ConsumeOneOutboundMessage(const FOutboundMessage& Out
 					  [FindNodeWithLog](const FGuid &Id) {return FindNodeWithLog(Id);},
 					  [this](const FGuid &Id) {return KnownNodes.Find(Id);} );
 
-	const bool bCanConsume = Algo::AllOf(Recipients, [](FNodeInfo *Node)
+	const bool bCanConsume = Algo::AllOf(Recipients, [this](FNodeInfo *Node)
 		{
-			const bool bWorkQueueIsNotFull = !Node->WorkQueue.IsFull();
-			if (!bWorkQueueIsNotFull)
-			{
-				UE_LOG(LogUdpMessaging, Warning,
-					   TEXT("Work queue for node %s is full. Send queue is congested. Will retry on next tick."), *Node->Endpoint.ToString());
-			}
-			return bWorkQueueIsNotFull;
+			return Node->CanCommitToWorkQueue(CurrentTime);
 		});
 
 	const bool bIsReliable = EnumHasAnyFlags(OutboundMessage.MessageFlags, EMessageFlags::Reliable);
diff --git a/Engine/Plugins/Messaging/UdpMessaging/Source/UdpMessaging/Private/Transport/UdpMessageProcessor.h b/Engine/Plugins/Messaging/UdpMessaging/Source/UdpMessaging/Private/Transport/UdpMessageProcessor.h
--- a/Engine/Plugins/Messaging/UdpMessaging/Source/UdpMessaging/Private/Transport/UdpMessageProcessor.h
+++ b/Engine/Plugins/Messaging/UdpMessaging/Source/UdpMessaging/Private/Transport/UdpMessageProcessor.h
@@ -151,6 +151,12 @@ class FUdpMessageProcessor
 		/** Holds of queue of MessageIds to send. They are processed in round-robin fashion. */
 		TUdpCircularQueue<int32> WorkQueue;
 
+		/** Last time we issued a warning about work queues. */
+		FDateTime LastWorkQueueFullMessage;
+
+		/** The number of seconds we have been in a send full state. */
+		double SecondsInFullQueueState = 0;
+
 		/** A map from sequence id to information about what was sent. The size of this map is the size of our sending window. */
 		TMap<uint64, FSentData> InflightSegments;
 
@@ -169,6 +175,33 @@ class FUdpMessageProcessor
 		/** Default constructor. */
 		FNodeInfo();
 
+		/**
+		 * Check the work queue buffer and note a warning if we are overcommitted.
+		 */
+		bool CanCommitToWorkQueue(const FDateTime& InCurrentTime)
+		{
+			const bool bWorkQueueFull = WorkQueue.IsFull();
+			if (bWorkQueueFull)
+			{
+				const double ReportingIntervalInSeconds = 1;
+				double Seconds = (InCurrentTime - LastWorkQueueFullMessage).GetTotalSeconds();
+				if (Seconds > ReportingIntervalInSeconds)
+				{
+					SecondsInFullQueueState += Seconds;
+					UE_LOG(LogUdpMessaging, Warning,
+						   TEXT("Work queue for node %s is full. Send queue has been congested for %s seconds."),
+						   *Endpoint.ToString(), SecondsInFullQueueState);
+					LastWorkQueueFullMessage = InCurrentTime;
+				}
+			}
+			else
+			{
+				LastWorkQueueFullMessage = InCurrentTime;
+				SecondsInFullQueueState = 0;
+			}
+			return !bWorkQueueFull;
+		}
+
 		/**
 		 * Remap a MessageId and SegmentId pair into a 64 bit number using Szudzik pairing calculation.
 		 */
