Use proper format specifier in the UE_LOG statement.

diff --git a/Engine/Plugins/Messaging/UdpMessaging/Source/UdpMessaging/Private/Transport/UdpMessageProcessor.h b/Engine/Plugins/Messaging/UdpMessaging/Source/UdpMessaging/Private/Transport/UdpMessageProcessor.h
--- a/Engine/Plugins/Messaging/UdpMessaging/Source/UdpMessaging/Private/Transport/UdpMessageProcessor.h
+++ b/Engine/Plugins/Messaging/UdpMessaging/Source/UdpMessaging/Private/Transport/UdpMessageProcessor.h
@@ -189,7 +189,7 @@ class FUdpMessageProcessor
 				{
 					SecondsInFullQueueState += Seconds;
 					UE_LOG(LogUdpMessaging, Warning,
-						   TEXT("Work queue for node %s is full. Send queue has been congested for %s seconds."),
+						   TEXT("Work queue for node %s is full. Send queue has been congested for %.3f seconds."),
 						   *Endpoint.ToString(), SecondsInFullQueueState);
 					LastWorkQueueFullMessage = InCurrentTime;
 				}
