Add an overflow buffer for reliably sent messages so that fully sent segments waiting for ack do not block new messages from being sent.

diff --git a/Engine/Plugins/Messaging/UdpMessaging/Source/UdpMessaging/Private/Tests/UdpMessageTransportTest.cpp b/Engine/Plugins/Messaging/UdpMessaging/Source/UdpMessaging/Private/Tests/UdpMessageTransportTest.cpp
--- a/Engine/Plugins/Messaging/UdpMessaging/Source/UdpMessaging/Private/Tests/UdpMessageTransportTest.cpp
+++ b/Engine/Plugins/Messaging/UdpMessaging/Source/UdpMessaging/Private/Tests/UdpMessageTransportTest.cpp
@@ -7,11 +7,12 @@
 #include "Interfaces/IPv4/IPv4Endpoint.h"
 #include "IMessageContext.h"
 #include "IMessageTransportHandler.h"
+#include "Templates/SharedPointerInternals.h"
 #include "Transport/UdpMessageTransport.h"
 #include "Tests/UdpMessagingTestTypes.h"
 
 
-IMPLEMENT_SIMPLE_AUTOMATION_TEST(FUdpMessageTransportTest, "System.Core.Messaging.Transports.Udp.UdpMessageTransport (takes ~2 minutes!)", EAutomationTestFlags::Disabled | EAutomationTestFlags::ApplicationContextMask | EAutomationTestFlags::EngineFilter)
+IMPLEMENT_SIMPLE_AUTOMATION_TEST(FUdpMessageTransportTest, "System.Core.Messaging.Transports.Udp.UdpMessageTransport", EAutomationTestFlags::ApplicationContextMask | EAutomationTestFlags::EngineFilter)
 
 
 class FUdpMessageTransportTestState
@@ -78,6 +79,12 @@ public:
 	virtual void ReceiveTransportMessage(const TSharedRef<IMessageContext, ESPMode::ThreadSafe>& Context, const FGuid& NodeId) override
 	{
 		FPlatformAtomics::InterlockedIncrement(&NumReceivedMessages);
+		LastMessageContext = Context;
+	}
+
+	TSharedPtr<IMessageContext, ESPMode::ThreadSafe> GetLastMessage() const
+	{
+		return LastMessageContext;
 	}
 
 private:
@@ -86,6 +93,7 @@ private:
 	TArray<FGuid> LostNodes;
 	int32 NumReceivedMessages;
 	TSharedPtr<IMessageTransport, ESPMode::ThreadSafe> Transport;
+	TSharedPtr<IMessageContext, ESPMode::ThreadSafe> LastMessageContext;
 };
 
 
@@ -144,7 +152,8 @@ bool FUdpMessageTransportTest::RunTest(const FString& Parameters)
 		for (int32 Count = 0; Count < NumTestMessages; ++Count)
 		{
 			FUdpMockMessage* Message = new FUdpMockMessage(MessageSize);
-			TSharedRef<IMessageContext, ESPMode::ThreadSafe> Context = MakeShareable(new FUdpMockMessageContext(Message, StartTime));
+			// Send messages reliably otherwise they may not send due to congestion and potentially dropped by the processor.
+			TSharedRef<IMessageContext, ESPMode::ThreadSafe> Context = MakeShareable(new FUdpMockMessageContext(Message, StartTime, EMessageFlags::Reliable));
 
 			Transport1.Publish(Context);
 		}
@@ -160,6 +169,48 @@ bool FUdpMessageTransportTest::RunTest(const FString& Parameters)
 		TestTrue(TEXT("All sent messages must have been received"), Transport2.GetNumReceivedMessages() == NumTestMessages);
 	}
 
+	if (Transport2.GetNumReceivedMessages() != NumTestMessages)
+	{
+		return false;
+	}
+
+	{
+		const FDateTime StartTime = FDateTime::UtcNow();
+
+		// 100 MB Test message.
+		const int32 LargeMessageSize = 1024*1024*100;
+		FUdpMockMessage* MockMessageToSend = new FUdpMockMessage(LargeMessageSize);
+		// Send messages reliably otherwise they may not send due to congestion and potentially dropped by the processor.
+		TSharedRef<IMessageContext, ESPMode::ThreadSafe> Context = MakeShareable(new FUdpMockMessageContext(MockMessageToSend, StartTime, EMessageFlags::Reliable));
+
+		const uint32 SenderCRC = MockMessageToSend->ComputeCRC();
+		Transport1.Publish(Context);
+
+		while ((Transport2.GetNumReceivedMessages() < NumTestMessages+1) && ((FDateTime::UtcNow() - StartTime) < FTimespan::FromSeconds(120.0)))
+		{
+			FPlatformProcess::Sleep(0.0f);
+		}
+
+		if ((Transport2.GetNumReceivedMessages() == NumTestMessages+1))
+		{
+			TSharedPtr<IMessageContext, ESPMode::ThreadSafe> ReceivedContext = Transport2.GetLastMessage();
+			const UScriptStruct* MessageTypeInfo = Context->GetMessageTypeInfo().Get();
+
+			const bool bIsAMockMessage = MessageTypeInfo->IsChildOf(FUdpMockMessage::StaticStruct());
+			TestTrue(TEXT("Is Udp Mock Message"), bIsAMockMessage);
+
+			if (bIsAMockMessage)
+			{
+				const FUdpMockMessage* ReceivedMockMessage = static_cast<const FUdpMockMessage*>(Context->GetMessage());
+				TestTrue(TEXT("Mock message size matches"), ReceivedMockMessage->Data.Num() == MockMessageToSend->Data.Num());
+				const uint32 ReceiverCRC = ReceivedMockMessage->ComputeCRC();
+				TestTrue(TEXT("CRC Match for large message send. "), ReceiverCRC == SenderCRC);
+			}
+
+			AddInfo(TEXT("Completed large message test."));
+		}
+	}
+
 	return true;
 }
 
diff --git a/Engine/Plugins/Messaging/UdpMessaging/Source/UdpMessaging/Private/Tests/UdpMessagingTestTypes.h b/Engine/Plugins/Messaging/UdpMessaging/Source/UdpMessaging/Private/Tests/UdpMessagingTestTypes.h
--- a/Engine/Plugins/Messaging/UdpMessaging/Source/UdpMessaging/Private/Tests/UdpMessagingTestTypes.h
+++ b/Engine/Plugins/Messaging/UdpMessaging/Source/UdpMessaging/Private/Tests/UdpMessagingTestTypes.h
@@ -3,6 +3,8 @@
 #pragma once
 
 #include "CoreMinimal.h"
+
+#include "Misc/Crc.h"
 #include "UObject/ObjectMacros.h"
 #include "IMessageContext.h"
 #include "IMessageAttachment.h"
@@ -25,6 +27,24 @@ struct FUdpMockMessage
 	{
 		Data.AddUninitialized(DataSize);
 	}
+
+	bool Serialize(FArchive& Ar)
+	{
+		int32 Num = Data.Num();
+		Ar << Num;
+		if (Ar.IsLoading())
+		{
+			Data.AddUninitialized(Num);
+		}
+		Ar.Serialize(Data.GetData(), Num);
+		return true;
+	}
+
+	uint32 ComputeCRC() const
+	{
+		uint32 CRC = 0;
+		return FCrc::MemCrc32(Data.GetData(), Data.Num(), CRC);
+	}
 };
 
 
@@ -33,11 +53,11 @@ class FUdpMockMessageContext
 {
 public:
 
-	FUdpMockMessageContext(FUdpMockMessage* InMessage, const FDateTime& InTimeSent)
+	FUdpMockMessageContext(FUdpMockMessage* InMessage, const FDateTime& InTimeSent, EMessageFlags SentFlags = EMessageFlags::None)
 		: Expiration(FDateTime::MaxValue())
 		, Message(InMessage)
 		, Scope(EMessageScope::Network)
-		, Flags(EMessageFlags::None)
+		, Flags(SentFlags)
 		, SenderThread(ENamedThreads::AnyThread)
 		, TimeSent(InTimeSent)
 		, TypeInfo(FUdpMockMessage::StaticStruct())
diff --git a/Engine/Plugins/Messaging/UdpMessaging/Source/UdpMessaging/Private/Transport/UdpMessageProcessor.cpp b/Engine/Plugins/Messaging/UdpMessaging/Source/UdpMessaging/Private/Transport/UdpMessageProcessor.cpp
--- a/Engine/Plugins/Messaging/UdpMessaging/Source/UdpMessaging/Private/Transport/UdpMessageProcessor.cpp
+++ b/Engine/Plugins/Messaging/UdpMessaging/Source/UdpMessaging/Private/Transport/UdpMessageProcessor.cpp
@@ -1036,12 +1036,12 @@ FSentSegmentInfo FUdpMessageProcessor::SendNextSegmentForMessageId(FNodeInfo& No
 	if (!FoundSegmenter)
 	{
 		// It's possible that a segmenter was fully ack and removed from the segmenter list.
-		return {};
+		return {MessageId};
 	}
 
 	TSharedPtr<FUdpMessageSegmenter>& Segmenter = *FoundSegmenter;
 
-	FSentSegmentInfo SentInfo;
+	FSentSegmentInfo SentInfo(MessageId);
 	if (Segmenter->IsInitialized() && Segmenter->NeedSending(CurrentTime))
 	{
 		FUdpMessageSegment::FDataChunk DataChunk = GetDataChunk(NodeInfo, Segmenter, MessageId);
@@ -1075,6 +1075,7 @@ FSentSegmentInfo FUdpMessageProcessor::SendNextSegmentForMessageId(FNodeInfo& No
 		SentInfo.SequenceNumber = ++NodeInfo.SequenceId;
 		SentInfo.bIsReliable = EnumHasAnyFlags(Segmenter->GetMessageFlags(), EMessageFlags::Reliable);
 		SentInfo.bRequiresRequeue = !Segmenter->IsSendingComplete() || SentInfo.bIsReliable;
+		SentInfo.bFullySent = Segmenter->IsSendingComplete();
 
 		Segmenter->UpdateSentTime(CurrentTime);
 
@@ -1094,6 +1095,7 @@ FSentSegmentInfo FUdpMessageProcessor::SendNextSegmentForMessageId(FNodeInfo& No
 	else
 	{
 		SentInfo.bRequiresRequeue = true;
+		SentInfo.bFullySent = Segmenter->IsSendingComplete();
 	}
 
 	NodeInfo.Statistics.TotalBytesSent += SentInfo.BytesSent;
@@ -1112,34 +1114,58 @@ int32 FUdpMessageProcessor::UpdateSegmenters(FNodeInfo& NodeInfo)
 		EUdpMessageSegments::Data		// Header.SegmentType
 	};
 
-
 	int32 MessageId;
 	int32 BytesSent = 0;
-	TArray<int32,TInlineAllocator<64>> MessagesForRequeue;
-	while (NodeInfo.CanSendSegments() && NodeInfo.WorkQueue.Dequeue(MessageId))
+	TArray<FSentSegmentInfo,TInlineAllocator<64>> ProcessedMessages;
+	auto SendOne = [this, &NodeInfo, &Header, &ProcessedMessages](int32 MessageId)
 	{
 		FSentSegmentInfo Info = SendNextSegmentForMessageId(NodeInfo, Header, MessageId);
-		if (Info.bSendSocketError)
+		ProcessedMessages.Emplace(MoveTemp(Info));
+	};
+	auto SocketError = [&ProcessedMessages]()
+	{
+		return ProcessedMessages.Top().bSendSocketError;
+	};
+
+	// Check for acknowledged segments first.
+	while (NodeInfo.CanSendSegments() && !NodeInfo.OverflowForPendingAck.IsEmpty())
+	{
+		SendOne(NodeInfo.OverflowForPendingAck.Pop());
+		if (SocketError())
 		{
 			return -1;
 		}
+	}
+	NodeInfo.OverflowForPendingAck.Reset();
 
-		if (Info.bRequiresRequeue)
+	// Process messages in the work queue.
+	while (NodeInfo.CanSendSegments() && NodeInfo.WorkQueue.Dequeue(MessageId))
+	{
+		SendOne(MessageId);
+		if (SocketError())
 		{
-			MessagesForRequeue.Add(MessageId);
+			return -1;
+		}
+	}
+
+	for (const FSentSegmentInfo& Info : ProcessedMessages)
+	{
+		if ((Info.bFullySent && Info.bRequiresRequeue) ||
+			(Info.bRequiresRequeue && NodeInfo.WorkQueue.IsFull()) )
+		{
+			NodeInfo.OverflowForPendingAck.Add(Info.MessageId);
+		}
+		else if (Info.bRequiresRequeue)
+		{
+			NodeInfo.WorkQueue.Enqueue(Info.MessageId);
 		}
 
 		BytesSent += Info.BytesSent;
 	}
-	
+
 	NodeInfo.Statistics.PacketsInFlight = NodeInfo.InflightSegments.Num();
 	NodeInfo.Statistics.BytesInflight = NodeInfo.InflightSegments.Num() * UDP_MESSAGING_SEGMENT_SIZE;
 
-	// Requeue for the next round.
-	for (int32 RequeueMessageId : MessagesForRequeue)
-	{
-		NodeInfo.WorkQueue.Enqueue(RequeueMessageId);
-	}
 	return BytesSent;
 }
 
diff --git a/Engine/Plugins/Messaging/UdpMessaging/Source/UdpMessaging/Private/Transport/UdpMessageProcessor.h b/Engine/Plugins/Messaging/UdpMessaging/Source/UdpMessaging/Private/Transport/UdpMessageProcessor.h
--- a/Engine/Plugins/Messaging/UdpMessaging/Source/UdpMessaging/Private/Transport/UdpMessageProcessor.h
+++ b/Engine/Plugins/Messaging/UdpMessaging/Source/UdpMessaging/Private/Transport/UdpMessageProcessor.h
@@ -81,18 +81,26 @@ struct FSentData
 /** Returned by the segment processor. It provides details on what was sent and how it was sent. */
 struct FSentSegmentInfo
 {
+	FSentSegmentInfo() = delete;
+	FSentSegmentInfo(int32 InMessageId)
+		: MessageId(InMessageId)
+	{
+	}
+
+	int32 MessageId;
 	uint64 SequenceNumber = 0;
 	uint32 BytesSent = 0;
 
 	bool bIsReliable = false;
 	bool bRequiresRequeue = false;
 	bool bSendSocketError = false;
+	bool bFullySent = false;
 
 	/** Converts FSentSegmentInfo into a FSentData struct. */
-	FSentData AsSentData(int32 MessageId, uint32 SegmentNumber, const FDateTime& CurrentTime)
+	FSentData AsSentData(int32 InMessageId, uint32 SegmentNumber, const FDateTime& CurrentTime)
 	{
 		return {
-			MessageId,
+			InMessageId,
 			SegmentNumber,
 			SequenceNumber,
 			bIsReliable,
@@ -151,6 +159,9 @@ class FUdpMessageProcessor
 		/** Holds of queue of MessageIds to send. They are processed in round-robin fashion. */
 		TUdpCircularQueue<int32> WorkQueue;
 
+		/* Holds the list of messages that are still pending acknowledgement or overflow if overcommitted on work queue. */
+		TArray<int32> OverflowForPendingAck;
+
 		/** Last time we issued a warning about work queues. */
 		FDateTime LastWorkQueueFullMessage;
 
diff --git a/Engine/Plugins/Messaging/UdpMessaging/Source/UdpMessaging/Private/Transport/UdpMessageTransport.cpp b/Engine/Plugins/Messaging/UdpMessaging/Source/UdpMessaging/Private/Transport/UdpMessageTransport.cpp
--- a/Engine/Plugins/Messaging/UdpMessaging/Source/UdpMessaging/Private/Transport/UdpMessageTransport.cpp
+++ b/Engine/Plugins/Messaging/UdpMessaging/Source/UdpMessaging/Private/Transport/UdpMessageTransport.cpp
@@ -53,6 +53,69 @@ TAutoConsoleVariable<bool> CVarEndpointDenyListEnabled(
 	ECVF_Default
 );
 
+namespace UE::UdpMessageTransport::Private
+{
+
+struct FDenyList
+{
+	struct FDenyCandidate
+	{
+		int32 EndpointFailureCount = 0;
+		FDateTime LastFailTime = FDateTime::UtcNow();
+	};
+
+	FDenyList() :
+		ClearDenyList(TEXT("MessageBus.UDP.ClearDenyList"), TEXT("Clear the UDP socket deny list."),
+					  FConsoleCommandDelegate::CreateRaw(this, &FDenyList::DoClearDenyCandidateList))
+	{
+	}
+
+	void DoClearDenyCandidateList()
+	{
+		FScopeLock DenyListLock(&DenyListCS);
+		DenyCandidateList.Reset();
+	}
+
+	bool ShouldAllowEndpoint(const FGuid& EndpointNodeId) const
+	{
+		FScopeLock DenyListLock(&DenyListCS);
+		const FDenyCandidate* DenyCandidate = DenyCandidateList.Find(EndpointNodeId);
+		if (DenyCandidate && DenyCandidate->EndpointFailureCount > CVarMaxRetriesForBadEndpoint.GetValueOnAnyThread())
+		{
+			return false;
+		}
+		return true;
+	}
+
+	void HandleEndpointError(const FGuid& EndpointId)
+	{
+		FScopeLock DenyListLock(&DenyListCS);
+		FDenyCandidate& DenyCandidate = DenyCandidateList.FindOrAdd(EndpointId);
+		FDateTime CurrentTime = FDateTime::UtcNow();
+		if ((CurrentTime - DenyCandidate.LastFailTime).GetSeconds() < CVarBadEndpointPeriod.GetValueOnAnyThread())
+		{
+			DenyCandidate.EndpointFailureCount++;
+		}
+		else
+		{
+			DenyCandidate.EndpointFailureCount = 0;
+		}
+		DenyCandidate.LastFailTime = CurrentTime;
+	}
+
+	/** Mutex protecting access to the DenyList map.. */
+	mutable FCriticalSection DenyListCS;
+
+	/** Deny list of node ids not allowed to talk to UDP processor */
+	TMap<FGuid, FDenyCandidate> DenyCandidateList;
+
+	/** Console command to reset the deny list */
+	FAutoConsoleCommand ClearDenyList;
+};
+
+static FDenyList GlobalDenyCandidateTable;
+}
+
 FUdpMessageTransport::FUdpMessageTransport(const FIPv4Endpoint& InUnicastEndpoint, const FIPv4Endpoint& InMulticastEndpoint, TArray<FIPv4Endpoint> InStaticEndpoints, TArray<FIPv4Endpoint> InExcludedEndpoints, uint8 InMulticastTtl)
 	: MessageProcessor(nullptr)
 	, MulticastEndpoint(InMulticastEndpoint)
@@ -68,8 +131,6 @@ FUdpMessageTransport::FUdpMessageTransport(const FIPv4Endpoint& InUnicastEndpoin
 #endif
 	, StaticEndpoints(MoveTemp(InStaticEndpoints))
 	, ExcludedEndpoints(MoveTemp(InExcludedEndpoints))
-	, ClearDenyList(TEXT("MessageBus.UDP.ClearDenyList"), TEXT("Clear the UDP socket deny list."),
-					FConsoleCommandDelegate::CreateRaw(this, &FUdpMessageTransport::DoClearDenyCandidateList))
 {
 }
 
@@ -388,8 +449,7 @@ bool FUdpMessageTransport::HandleProcessorEndpointCheck(const FGuid& EndpointNod
 {
 	if (CVarEndpointDenyListEnabled.GetValueOnAnyThread())
 	{
-		FDenyCandidate* DenyCandidate = DenyCandidateList.Find(EndpointNodeId);
-		if (DenyCandidate && DenyCandidate->EndpointFailureCount > CVarMaxRetriesForBadEndpoint.GetValueOnAnyThread())
+		if (!UE::UdpMessageTransport::Private::GlobalDenyCandidateTable.ShouldAllowEndpoint(EndpointNodeId))
 		{
 			return false;
 		}
@@ -406,24 +466,9 @@ bool FUdpMessageTransport::HandleProcessorEndpointCheck(const FGuid& EndpointNod
 	return true;
 }
 
-void FUdpMessageTransport::DoClearDenyCandidateList()
-{
-	DenyCandidateList.Reset();
-}
-
 void FUdpMessageTransport::HandleEndpointCommunicationError(const FGuid& EndpointId, const FIPv4Endpoint& /*Unused EndpointIdAddress*/)
 {
-	FDenyCandidate& DenyCandidate = DenyCandidateList.FindOrAdd(EndpointId);
-	FDateTime CurrentTime = FDateTime::UtcNow();
-	if ((CurrentTime - DenyCandidate.LastFailTime).GetSeconds() < CVarBadEndpointPeriod.GetValueOnAnyThread())
-	{
-		DenyCandidate.EndpointFailureCount++;
-	}
-	else
-	{
-		DenyCandidate.EndpointFailureCount = 0;
-	}
-	DenyCandidate.LastFailTime = CurrentTime;
+	UE::UdpMessageTransport::Private::GlobalDenyCandidateTable.HandleEndpointError(EndpointId);
 }
 
 void FUdpMessageTransport::HandleProcessorError()
diff --git a/Engine/Plugins/Messaging/UdpMessaging/Source/UdpMessaging/Private/Transport/UdpMessageTransport.h b/Engine/Plugins/Messaging/UdpMessaging/Source/UdpMessaging/Private/Transport/UdpMessageTransport.h
--- a/Engine/Plugins/Messaging/UdpMessaging/Source/UdpMessaging/Private/Transport/UdpMessageTransport.h
+++ b/Engine/Plugins/Messaging/UdpMessaging/Source/UdpMessaging/Private/Transport/UdpMessageTransport.h
@@ -173,21 +173,9 @@ private:
 	FSocket* UnicastSocket;
 #endif
 
-	struct FDenyCandidate
-	{
-		int32 EndpointFailureCount = 0;
-		FDateTime LastFailTime = FDateTime::UtcNow();
-	};
-
-	/** Deny list of node ids not allowed to talk to UDP processor */
-	TMap<FGuid, FDenyCandidate> DenyCandidateList;
-
 	/** Holds the static endpoints. */
 	TSet<FIPv4Endpoint> StaticEndpoints;
 
 	/** Excluded list of IP addresses not allowed to talk to UDP processor */
 	TArray<FIPv4Endpoint> ExcludedEndpoints;
-
-	/** Console command to reset the deny list */
-	FAutoConsoleCommand ClearDenyList;
 };
