Fix invalid allocation of message in UdpMessageTransportTest.cpp

diff --git a/Engine/Plugins/Messaging/UdpMessaging/Source/UdpMessaging/Private/Tests/UdpMessageTransportTest.cpp b/Engine/Plugins/Messaging/UdpMessaging/Source/UdpMessaging/Private/Tests/UdpMessageTransportTest.cpp
--- a/Engine/Plugins/Messaging/UdpMessaging/Source/UdpMessaging/Private/Tests/UdpMessageTransportTest.cpp
+++ b/Engine/Plugins/Messaging/UdpMessaging/Source/UdpMessaging/Private/Tests/UdpMessageTransportTest.cpp
@@ -10,7 +10,7 @@
 #include "Templates/SharedPointerInternals.h"
 #include "Transport/UdpMessageTransport.h"
 #include "Tests/UdpMessagingTestTypes.h"
-
+#include "MessageEndpoint.h"
 
 IMPLEMENT_SIMPLE_AUTOMATION_TEST(FUdpMessageTransportTest, "System.Core.Messaging.Transports.Udp.UdpMessageTransport", EAutomationTestFlags::ApplicationContextMask | EAutomationTestFlags::EngineFilter)
 
@@ -151,7 +151,7 @@ bool FUdpMessageTransportTest::RunTest(const FString& Parameters)
 
 		for (int32 Count = 0; Count < NumTestMessages; ++Count)
 		{
-			FUdpMockMessage* Message = new FUdpMockMessage(MessageSize);
+			FUdpMockMessage* Message = FMessageEndpoint::MakeMessage<FUdpMockMessage>(MessageSize);
 			// Send messages reliably otherwise they may not send due to congestion and potentially dropped by the processor.
 			TSharedRef<IMessageContext, ESPMode::ThreadSafe> Context = MakeShareable(new FUdpMockMessageContext(Message, StartTime, EMessageFlags::Reliable));
 
@@ -179,7 +179,7 @@ bool FUdpMessageTransportTest::RunTest(const FString& Parameters)
 
 		// 100 MB Test message.
 		const int32 LargeMessageSize = 1024*1024*100;
-		FUdpMockMessage* MockMessageToSend = new FUdpMockMessage(LargeMessageSize);
+		FUdpMockMessage* MockMessageToSend = FMessageEndpoint::MakeMessage<FUdpMockMessage>(LargeMessageSize);
 		// Send messages reliably otherwise they may not send due to congestion and potentially dropped by the processor.
 		TSharedRef<IMessageContext, ESPMode::ThreadSafe> Context = MakeShareable(new FUdpMockMessageContext(MockMessageToSend, StartTime, EMessageFlags::Reliable));
 
