UdpMessageTransportTest: Fix unsynchronized multithreaded access to LastMessageContext. LastMessageContext is a threadsafe TSharedPtr, but this only protects the counter value being pointed to. The TSharedPtr structure itself is not threadsafe and must be guarded if used from multiple threads.
#rnx

diff --git a/Engine/Plugins/Messaging/UdpMessaging/Source/UdpMessaging/Private/Tests/UdpMessageTransportTest.cpp b/Engine/Plugins/Messaging/UdpMessaging/Source/UdpMessaging/Private/Tests/UdpMessageTransportTest.cpp
--- a/Engine/Plugins/Messaging/UdpMessaging/Source/UdpMessaging/Private/Tests/UdpMessageTransportTest.cpp
+++ b/Engine/Plugins/Messaging/UdpMessaging/Source/UdpMessaging/Private/Tests/UdpMessageTransportTest.cpp
@@ -6,6 +6,8 @@
 #include "Transport/UdpMessageTransport.h"
 #include "Tests/UdpMessagingTestTypes.h"
 #include "MessageEndpoint.h"
+#include "HAL/CriticalSection.h"
+#include "Misc/ScopeLock.h"
 
 IMPLEMENT_SIMPLE_AUTOMATION_TEST(FUdpMessageTransportTest, "System.Core.Messaging.Transports.Udp.UdpMessageTransport", EAutomationTestFlags::ApplicationContextMask | EAutomationTestFlags::EngineFilter)
 
@@ -74,11 +76,13 @@ public:
 	virtual void ReceiveTransportMessage(const TSharedRef<IMessageContext, ESPMode::ThreadSafe>& Context, const FGuid& NodeId) override
 	{
 		FPlatformAtomics::InterlockedIncrement(&NumReceivedMessages);
+		FScopeLock LastMessageContextScopeLock(&LastMessageContextLock);
 		LastMessageContext = Context;
 	}
 
 	TSharedPtr<IMessageContext, ESPMode::ThreadSafe> GetLastMessage() const
 	{
+		FScopeLock LastMessageContextScopeLock(&LastMessageContextLock);
 		return LastMessageContext;
 	}
 
@@ -87,6 +91,7 @@ private:
 	TArray<FGuid> DiscoveredNodes;
 	TArray<FGuid> LostNodes;
 	int32 NumReceivedMessages;
+	mutable FCriticalSection LastMessageContextLock;
 	TSharedPtr<IMessageTransport, ESPMode::ThreadSafe> Transport;
 	TSharedPtr<IMessageContext, ESPMode::ThreadSafe> LastMessageContext;
 };
