Removing Private include paths from UdpMessaging.

diff --git a/Engine/Plugins/Messaging/UdpMessaging/Source/UdpMessaging/Private/Tests/UdpMessageSegmenterTest.cpp b/Engine/Plugins/Messaging/UdpMessaging/Source/UdpMessaging/Private/Tests/UdpMessageSegmenterTest.cpp
--- a/Engine/Plugins/Messaging/UdpMessaging/Source/UdpMessaging/Private/Tests/UdpMessageSegmenterTest.cpp
+++ b/Engine/Plugins/Messaging/UdpMessaging/Source/UdpMessaging/Private/Tests/UdpMessageSegmenterTest.cpp
@@ -4,7 +4,7 @@
 #include "Containers/ArrayBuilder.h"
 
 #include "UdpMessagingPrivate.h"
-#include "UdpMessagingSettings.h"
+#include "Shared/UdpMessagingSettings.h"
 #include "Transport/UdpSerializedMessage.h"
 #include "Transport/UdpMessageSegmenter.h"
 
diff --git a/Engine/Plugins/Messaging/UdpMessaging/Source/UdpMessaging/Private/Tests/UdpSerializeMessageTaskTest.cpp b/Engine/Plugins/Messaging/UdpMessaging/Source/UdpMessaging/Private/Tests/UdpSerializeMessageTaskTest.cpp
--- a/Engine/Plugins/Messaging/UdpMessaging/Source/UdpMessaging/Private/Tests/UdpSerializeMessageTaskTest.cpp
+++ b/Engine/Plugins/Messaging/UdpMessaging/Source/UdpMessaging/Private/Tests/UdpSerializeMessageTaskTest.cpp
@@ -3,7 +3,7 @@
 #include "Misc/AutomationTest.h"
 
 #include "UdpMessagingPrivate.h"
-#include "UdpMessagingSettings.h"
+#include "Shared/UdpMessagingSettings.h"
 #include "MessageEndpoint.h"
 #include "Transport/UdpSerializedMessage.h"
 #include "Transport/UdpSerializeMessageTask.h"
diff --git a/Engine/Plugins/Messaging/UdpMessaging/Source/UdpMessaging/Private/Transport/UdpDeserializedMessage.cpp b/Engine/Plugins/Messaging/UdpMessaging/Source/UdpMessaging/Private/Transport/UdpDeserializedMessage.cpp
--- a/Engine/Plugins/Messaging/UdpMessaging/Source/UdpMessaging/Private/Transport/UdpDeserializedMessage.cpp
+++ b/Engine/Plugins/Messaging/UdpMessaging/Source/UdpMessaging/Private/Transport/UdpDeserializedMessage.cpp
@@ -8,7 +8,7 @@
 #include "StructDeserializer.h"
 #include "Serialization/MemoryReader.h"
 #include "Transport/UdpReassembledMessage.h"
-#include "UdpMessagingSettings.h"
+#include "Shared/UdpMessagingSettings.h"
 #include "UdpMessagingTracing.h"
 
 /* FUdpDeserializedMessage structors
diff --git a/Engine/Plugins/Messaging/UdpMessaging/Source/UdpMessaging/Private/Transport/UdpMessageBeacon.cpp b/Engine/Plugins/Messaging/UdpMessaging/Source/UdpMessaging/Private/Transport/UdpMessageBeacon.cpp
--- a/Engine/Plugins/Messaging/UdpMessaging/Source/UdpMessaging/Private/Transport/UdpMessageBeacon.cpp
+++ b/Engine/Plugins/Messaging/UdpMessaging/Source/UdpMessaging/Private/Transport/UdpMessageBeacon.cpp
@@ -8,7 +8,7 @@
 #include "HAL/RunnableThread.h"
 #include "Serialization/ArrayWriter.h"
 #include "Sockets.h"
-#include "UdpMessageSegment.h"
+#include "Shared/UdpMessageSegment.h"
 
 /* FUdpMessageHelloSender static initialization
  *****************************************************************************/
diff --git a/Engine/Plugins/Messaging/UdpMessaging/Source/UdpMessaging/Private/Transport/UdpMessageProcessor.cpp b/Engine/Plugins/Messaging/UdpMessaging/Source/UdpMessaging/Private/Transport/UdpMessageProcessor.cpp
--- a/Engine/Plugins/Messaging/UdpMessaging/Source/UdpMessaging/Private/Transport/UdpMessageProcessor.cpp
+++ b/Engine/Plugins/Messaging/UdpMessaging/Source/UdpMessaging/Private/Transport/UdpMessageProcessor.cpp
@@ -15,7 +15,7 @@
 #include "Transport/UdpReassembledMessage.h"
 #include "Transport/UdpSerializedMessage.h"
 #include "Transport/UdpSerializeMessageTask.h"
-#include "UdpMessageSegment.h"
+#include "Shared/UdpMessageSegment.h"
 
 
 /* FUdpMessageHelloSender static initialization
diff --git a/Engine/Plugins/Messaging/UdpMessaging/Source/UdpMessaging/Private/Transport/UdpSerializeMessageTask.cpp b/Engine/Plugins/Messaging/UdpMessaging/Source/UdpMessaging/Private/Transport/UdpSerializeMessageTask.cpp
--- a/Engine/Plugins/Messaging/UdpMessaging/Source/UdpMessaging/Private/Transport/UdpSerializeMessageTask.cpp
+++ b/Engine/Plugins/Messaging/UdpMessaging/Source/UdpMessaging/Private/Transport/UdpSerializeMessageTask.cpp
@@ -9,7 +9,7 @@
 
 #include "UdpMessagingPrivate.h"
 #include "Transport/UdpSerializedMessage.h"
-#include "UdpMessagingSettings.h"
+#include "Shared/UdpMessagingSettings.h"
 #include "UdpMessagingTracing.h"
 
 namespace UdpSerializeMessageTaskDetails
diff --git a/Engine/Plugins/Messaging/UdpMessaging/Source/UdpMessaging/Private/UdpMessagingModule.cpp b/Engine/Plugins/Messaging/UdpMessaging/Source/UdpMessaging/Private/UdpMessagingModule.cpp
--- a/Engine/Plugins/Messaging/UdpMessaging/Source/UdpMessaging/Private/UdpMessagingModule.cpp
+++ b/Engine/Plugins/Messaging/UdpMessaging/Source/UdpMessaging/Private/UdpMessagingModule.cpp
@@ -1,6 +1,6 @@
 // Copyright Epic Games, Inc. All Rights Reserved.
 
-#include "UdpMessageProcessor.h"
+#include "Transport/UdpMessageProcessor.h"
 
 #include "MessageBridgeBuilder.h"
 #include "Misc/CoreDelegates.h"
diff --git a/Engine/Plugins/Messaging/UdpMessaging/Source/UdpMessaging/UdpMessaging.Build.cs b/Engine/Plugins/Messaging/UdpMessaging/Source/UdpMessaging/UdpMessaging.Build.cs
--- a/Engine/Plugins/Messaging/UdpMessaging/Source/UdpMessaging/UdpMessaging.Build.cs
+++ b/Engine/Plugins/Messaging/UdpMessaging/Source/UdpMessaging/UdpMessaging.Build.cs
@@ -29,15 +29,6 @@ namespace UnrealBuildTool.Rules
 					"MessagingCommon",
 				});
 
-			PrivateIncludePaths.AddRange(
-				new string[] {
-					"UdpMessaging/Private",
-					"UdpMessaging/Private/Shared",
-					"UdpMessaging/Private/Transport",
-					"UdpMessaging/Private/Tunnel",
-					"UdpMessaging/Public/Shared",
-				});
-
 			if (Target.Type == TargetType.Editor)
 			{
 				DynamicallyLoadedModuleNames.AddRange(
