UdpMessaging: Fix potential null pointer dereference if multicast socket can't be created.

diff --git a/Engine/Plugins/Messaging/UdpMessaging/Source/UdpMessaging/Private/Transport/UdpMessageTransport.cpp b/Engine/Plugins/Messaging/UdpMessaging/Source/UdpMessaging/Private/Transport/UdpMessageTransport.cpp
--- a/Engine/Plugins/Messaging/UdpMessaging/Source/UdpMessaging/Private/Transport/UdpMessageTransport.cpp
+++ b/Engine/Plugins/Messaging/UdpMessaging/Source/UdpMessaging/Private/Transport/UdpMessageTransport.cpp
@@ -285,9 +285,11 @@ bool FUdpMessageTransport::StartTransport(IMessageTransportHandler& Handler)
 #endif
 		.WithReceiveBufferSize(UDP_MESSAGING_RECEIVE_BUFFER_SIZE);
 
-	UE::UdpMessaging::Private::JoinedToGroup(UnicastEndpoint, MulticastEndpoint, MulticastSocket);
-	
-	if (MulticastSocket == nullptr)
+	if (MulticastSocket != nullptr)
+	{
+		UE::UdpMessaging::Private::JoinedToGroup(UnicastEndpoint, MulticastEndpoint, MulticastSocket);
+	}
+	else
 	{
 		UE_LOG(LogUdpMessaging, Warning, TEXT("StartTransport failed to create multicast socket on %s, joined to %s with TTL %i"), *UnicastEndpoint.ToString(), *MulticastEndpoint.ToString(), MulticastTtl);
 
