Add a CVar control transmission delay and max retries for UDP message segmenter.

diff --git a/Engine/Plugins/Messaging/UdpMessaging/Source/UdpMessaging/Private/Transport/UdpMessageSegmenter.cpp b/Engine/Plugins/Messaging/UdpMessaging/Source/UdpMessaging/Private/Transport/UdpMessageSegmenter.cpp
--- a/Engine/Plugins/Messaging/UdpMessaging/Source/UdpMessaging/Private/Transport/UdpMessageSegmenter.cpp
+++ b/Engine/Plugins/Messaging/UdpMessaging/Source/UdpMessaging/Private/Transport/UdpMessageSegmenter.cpp
@@ -227,13 +227,23 @@ void FUdpMessageSegmenter::MarkForRetransmission()
 	PendingSendSegments.Init(true, PendingSendSegments.Num());
 	PendingSendSegmentsCount = PendingSendSegments.Num();
 	
-	// Note - we don't need to clear acknowledgments. If any segment is in transit and acknowledged after this 
+	// Note - we don't need to clear acknowledgments. If any segment is in transit and acknowledged after this
 	// call we'll do that and if possible stop the pending send, and if not we don't need to wait for the ack
 }
 
+TAutoConsoleVariable<int32> CVarSegmenterTimeout(
+	TEXT("MessageBus.UDP.SegmenterTimeout"),
+	100,
+	TEXT("The number of milliseconds to wait for a segment to be acknowledge before attempting to resend. Values are clamped between 10 and 1000"),
+	ECVF_Default
+);
+TAutoConsoleVariable<int32> CVarSegmenterMaxResends(
+	TEXT("MessageBus.UDP.SegmenterMaxResends"),
+	16,
+	TEXT("The number of attempts to resend data to an endpoint. Values are clamped between 1 and 100."),
+	ECVF_Default
+);
 
-const FTimespan FUdpMessageSegmenter::SendInterval = FTimespan::FromMilliseconds(100);
-const uint16    MaxNumResends = 16;
 bool FUdpMessageSegmenter::NeedSending(const FDateTime& CurrentTime)
 {
 	// still have outstanding segments
@@ -242,6 +252,7 @@ bool FUdpMessageSegmenter::NeedSending(const FDateTime& CurrentTime)
 		return true;
 	}
 
+	const uint16 MaxNumResends = FMath::Clamp(CVarSegmenterMaxResends.GetValueOnAnyThread(), 1, 100);
 	if (AreAcknowledgementsComplete() == false && SentNumber > MaxNumResends)
 	{
 		UE_LOG(LogUdpMessaging, Warning, TEXT("Gave up sending with %d outstanding acks."), AcknowledgeSegments.Num() - AcknowledgeSegmentsCount);
@@ -249,6 +260,8 @@ bool FUdpMessageSegmenter::NeedSending(const FDateTime& CurrentTime)
 		return false;
 	}
 
+	const int32 NumMs = FMath::Clamp(CVarSegmenterTimeout.GetValueOnAnyThread(), 10, 5000);
+	const FTimespan SendInterval = FTimespan::FromMilliseconds(NumMs);
 	if (AreAcknowledgementsComplete() == false
 		&& LastSentTime + SendInterval <= CurrentTime)
 	{
diff --git a/Engine/Plugins/Messaging/UdpMessaging/Source/UdpMessaging/Private/Transport/UdpMessageSegmenter.h b/Engine/Plugins/Messaging/UdpMessaging/Source/UdpMessaging/Private/Transport/UdpMessageSegmenter.h
--- a/Engine/Plugins/Messaging/UdpMessaging/Source/UdpMessaging/Private/Transport/UdpMessageSegmenter.h
+++ b/Engine/Plugins/Messaging/UdpMessaging/Source/UdpMessaging/Private/Transport/UdpMessageSegmenter.h
@@ -239,9 +239,6 @@ public:
 	bool IsMessageSerializationDone() const;
 
 private:
-	/** Defines the time interval for sending. */
-	static const FTimespan SendInterval;
-
 	/** temp hack to support new transport API. */
 	FArchive* MessageReader;
 
