Prevent UDP message processor from processing messages until PostDefault phase of plugin loading.

diff --git a/Engine/Plugins/Developer/Concert/ConcertSync/ConcertSyncServer/Source/ConcertSyncServer/Public/ConcertSyncServerLoop.inl b/Engine/Plugins/Developer/Concert/ConcertSync/ConcertSyncServer/Source/ConcertSyncServer/Public/ConcertSyncServerLoop.inl
--- a/Engine/Plugins/Developer/Concert/ConcertSync/ConcertSyncServer/Source/ConcertSyncServer/Public/ConcertSyncServerLoop.inl
+++ b/Engine/Plugins/Developer/Concert/ConcertSync/ConcertSyncServer/Source/ConcertSyncServer/Public/ConcertSyncServerLoop.inl
@@ -51,7 +51,10 @@ int32 ConcertSyncServerLoop(const TCHAR* CommandLine, const FConcertSyncServerLo
 	// Load Concert Sync plugins in default phase
 	IPluginManager::Get().LoadModulesForEnabledPlugins(ELoadingPhase::Default);
 
-	// Install graceful termination handler, this handles graceful CTRL+C shutdown, 
+	// Load Concert Sync plugins in post default phase.
+	IPluginManager::Get().LoadModulesForEnabledPlugins(ELoadingPhase::PostDefault);
+
+	// Install graceful termination handler, this handles graceful CTRL+C shutdown,
 	// but not CTRL+CLOSE, which will potentially still exit process before the main thread exits.
 	// Double CTRL+C signal will also cause process to terminate.
 	FPlatformMisc::SetGracefulTerminationHandler();
diff --git a/Engine/Plugins/Messaging/UdpMessaging/Source/UdpMessaging/Private/Transport/UdpMessageProcessor.cpp b/Engine/Plugins/Messaging/UdpMessaging/Source/UdpMessaging/Private/Transport/UdpMessageProcessor.cpp
--- a/Engine/Plugins/Messaging/UdpMessaging/Source/UdpMessaging/Private/Transport/UdpMessageProcessor.cpp
+++ b/Engine/Plugins/Messaging/UdpMessaging/Source/UdpMessaging/Private/Transport/UdpMessageProcessor.cpp
@@ -17,6 +17,7 @@
 #include "Transport/UdpSerializeMessageTask.h"
 #include "Shared/UdpMessageSegment.h"
 
+#include "Interfaces/IPluginManager.h"
 
 /* FUdpMessageHelloSender static initialization
  *****************************************************************************/
@@ -97,7 +98,6 @@ FUdpMessageProcessor::FNodeInfo::FNodeInfo()
 	ComputeWindowSize(0,0);
 }
 
-
 FUdpMessageProcessor::FUdpMessageProcessor(FSocket& InSocket, const FGuid& InNodeId, const FIPv4Endpoint& InMulticastEndpoint)
 	: Beacon(nullptr)
 	, LocalNodeId(InNodeId)
@@ -117,7 +117,15 @@ FUdpMessageProcessor::FUdpMessageProcessor(FSocket& InSocket, const FGuid& InNod
 	});
 
 	CurrentTime = FDateTime::UtcNow();
-	Thread = FRunnableThread::Create(this, TEXT("FUdpMessageProcessor"), 128 * 1024, TPri_AboveNormal, FPlatformAffinity::GetPoolThreadMask());
+	const ELoadingPhase::Type LoadingPhase = IPluginManager::Get().GetLastCompletedLoadingPhase();
+	if (LoadingPhase == ELoadingPhase::None || LoadingPhase < ELoadingPhase::PostDefault)
+	{
+		IPluginManager::Get().OnLoadingPhaseComplete().AddRaw(this, &FUdpMessageProcessor::OnPluginLoadingPhaseComplete);
+	}
+	else
+	{
+		StartThread();
+	}
 }
 
 
@@ -126,14 +134,11 @@ FUdpMessageProcessor::~FUdpMessageProcessor()
 	// shut down worker thread if it is still running
 	Thread->Kill();
 
-	delete Thread;
-	Thread = nullptr;
-
-	delete Beacon;
-	Beacon = nullptr;
+	Thread = {};
+	Beacon = {};
+	SocketSender = {};
 
-	delete SocketSender;
-	SocketSender = nullptr;
+	IPluginManager::Get().OnLoadingPhaseComplete().RemoveAll(this);
 
 	// remove all transport nodes
 	if (NodeLostDelegate.IsBound())
@@ -147,6 +152,20 @@ FUdpMessageProcessor::~FUdpMessageProcessor()
 	KnownNodes.Empty();
 }
 
+void FUdpMessageProcessor::StartThread()
+{
+	Thread = TUniquePtr<FRunnableThread>(FRunnableThread::Create(this, TEXT("FUdpMessageProcessor"), 128 * 1024, TPri_AboveNormal, FPlatformAffinity::GetPoolThreadMask()));
+}
+
+void FUdpMessageProcessor::OnPluginLoadingPhaseComplete(ELoadingPhase::Type LoadingPhase, bool bPhaseSuccessful)
+{
+	check(!Thread);
+	if (LoadingPhase == ELoadingPhase::PostDefault)
+	{
+		IPluginManager::Get().OnLoadingPhaseComplete().RemoveAll(this);
+		StartThread();
+	}
+}
 
 void FUdpMessageProcessor::AddStaticEndpoint(const FIPv4Endpoint& InEndpoint)
 {
@@ -308,8 +327,8 @@ bool FUdpMessageProcessor::Init()
 {
 	if (!bIsInitialized)
 	{
-		Beacon = new FUdpMessageBeacon(Socket, LocalNodeId, MulticastEndpoint);
-		SocketSender = new FUdpSocketSender(Socket, TEXT("FUdpMessageProcessor.Sender"));
+		Beacon = MakeUnique<FUdpMessageBeacon>(Socket, LocalNodeId, MulticastEndpoint);
+		SocketSender = MakeUnique<FUdpSocketSender>(Socket, TEXT("FUdpMessageProcessor.Sender"));
 
 		// Current protocol version 16
 		SupportedProtocolVersions.Add(UDP_MESSAGING_TRANSPORT_PROTOCOL_VERSION);
@@ -448,7 +467,7 @@ void FUdpMessageProcessor::ConsumeInboundSegments()
 	SCOPED_MESSAGING_TRACE(FUdpMessageProcessor_ConsumeInboundSegments);
 	FInboundSegment Segment;
 
-	FTimespan MaxWorkTimespan = ComputeMaxWorkTimespan(DeadHelloIntervals, Beacon);
+	FTimespan MaxWorkTimespan = ComputeMaxWorkTimespan(DeadHelloIntervals, Beacon.Get());
 	while (InboundSegments.Dequeue(Segment))
 	{
 		// quick hack for TTP# 247103
diff --git a/Engine/Plugins/Messaging/UdpMessaging/Source/UdpMessaging/Private/Transport/UdpMessageProcessor.h b/Engine/Plugins/Messaging/UdpMessaging/Source/UdpMessaging/Private/Transport/UdpMessageProcessor.h
--- a/Engine/Plugins/Messaging/UdpMessaging/Source/UdpMessaging/Private/Transport/UdpMessageProcessor.h
+++ b/Engine/Plugins/Messaging/UdpMessaging/Source/UdpMessaging/Private/Transport/UdpMessageProcessor.h
@@ -4,6 +4,8 @@
 
 #include "Containers/Queue.h"
 
+#include "ModuleDescriptor.h"
+
 #include "HAL/Runnable.h"
 
 #include "INetworkMessagingExtension.h"
@@ -824,6 +826,11 @@ protected:
 	virtual void Tick() override;
 
 private:
+	/** Start the message processor thread. */
+	void StartThread();
+
+	/** Delegate invoke when plugin phase has been completed. */
+	void OnPluginLoadingPhaseComplete(ELoadingPhase::Type LoadingPhase, bool bPhaseSuccessful);
 
 	/** Handles a communication error to a particular endpoint.*/
 	void HandleSocketError(const FNodeInfo& NodeInfo) const;
@@ -841,7 +848,7 @@ private:
 	TQueue<FOutboundMessage, EQueueMode::Mpsc> OutboundMessages;
 
 	/** Holds the hello sender. */
-	FUdpMessageBeacon* Beacon;
+	TUniquePtr<FUdpMessageBeacon> Beacon;
 
 	/** Holds the current time. */
 	FDateTime CurrentTime;
@@ -879,8 +886,8 @@ private:
 	/** Holds the network socket used to transport messages. */
 	FSocket* Socket;
 
-	/** Holds the socket sender. volatile pointer because used to validate thread shutdown. */
-	FUdpSocketSender* volatile SocketSender;
+	/** Holds the socket sender.*/
+	TUniquePtr<FUdpSocketSender> SocketSender;
 
 	/** Holds a flag indicating that the thread is stopping. */
 	bool bStopping;
@@ -889,7 +896,7 @@ private:
 	bool bIsInitialized;
 
 	/** Holds the thread object. */
-	FRunnableThread* Thread;
+	TUniquePtr<FRunnableThread> Thread;
 
 	/** Holds an event signaling that inbound messages need to be processed. */
 	TSharedPtr<FEvent, ESPMode::ThreadSafe> WorkEvent;
diff --git a/Engine/Plugins/Messaging/UdpMessaging/Source/UdpMessaging/UdpMessaging.Build.cs b/Engine/Plugins/Messaging/UdpMessaging/Source/UdpMessaging/UdpMessaging.Build.cs
--- a/Engine/Plugins/Messaging/UdpMessaging/Source/UdpMessaging/UdpMessaging.Build.cs
+++ b/Engine/Plugins/Messaging/UdpMessaging/Source/UdpMessaging/UdpMessaging.Build.cs
@@ -19,9 +19,10 @@ namespace UnrealBuildTool.Rules
 					"Json",
 					"Cbor",
 					"Messaging",
+					"Projects",
 					"Serialization",
 					"TraceLog",
-					"Sockets",
+					"Sockets"
 				});
 
 			PrivateIncludePathModuleNames.AddRange(
