[UDPMessaging] Direct link now works when offline

- GetLocalAdapterAddresses returns empty list when no network adapters are connected(i.e. on a PC with WiFi disconnect you network, making sure it's not reconnected). Subsequently the socket is not joined to multicast group.
The fix(not sure how proper) seems to help.

diff --git a/Engine/Plugins/Messaging/UdpMessaging/Source/UdpMessaging/Private/Transport/UdpMessageTransport.cpp b/Engine/Plugins/Messaging/UdpMessaging/Source/UdpMessaging/Private/Transport/UdpMessageTransport.cpp
--- a/Engine/Plugins/Messaging/UdpMessaging/Source/UdpMessaging/Private/Transport/UdpMessageTransport.cpp
+++ b/Engine/Plugins/Messaging/UdpMessaging/Source/UdpMessaging/Private/Transport/UdpMessageTransport.cpp
@@ -231,6 +231,13 @@ void JoinedToGroup(const FIPv4Endpoint& UnicastEndpoint, const FIPv4Endpoint& Mu
 		{
 			DoJoinMulticastGroup(MulticastAddr, LocalIp, MulticastSocket);
 		}
+
+		// GetLocalAdapterAddresses returns empty list when all network adapters are offline
+		if (LocapIps.Num() == 0)
+		{
+			bool bCanBindAll = false;
+			DoJoinMulticastGroup(MulticastAddr, ISocketSubsystem::Get(PLATFORM_SOCKETSUBSYSTEM)->GetLocalHostAddr(*GLog, bCanBindAll), MulticastSocket);
+		}
 	}
 	else
 	{
