Add support for meshing networks. When enabled, the a newly discovered node will be shared with existing nodes connected to the UDP layer.  This feature has to be enabled via a setting and designed for Multi-user networks where some topology information is pre-shared on switchboard launch.  This initial sharing will allow editors to completely mesh together without having to know the editor endpoint id, which is typically auto assigned a port.

diff --git a/Engine/Plugins/Messaging/UdpMessaging/Source/UdpMessaging/Private/Shared/UdpMessageSegment.h b/Engine/Plugins/Messaging/UdpMessaging/Source/UdpMessaging/Private/Shared/UdpMessageSegment.h
--- a/Engine/Plugins/Messaging/UdpMessaging/Source/UdpMessaging/Private/Shared/UdpMessageSegment.h
+++ b/Engine/Plugins/Messaging/UdpMessaging/Source/UdpMessaging/Private/Shared/UdpMessageSegment.h
@@ -46,6 +46,9 @@ enum class EUdpMessageSegments : uint8
 	/** Answers back to ping segment. */
 	Pong,
 
+	/** Used to share endpoints for fully meshing networks. */
+	Mesh
+
 	// New segment type needs to be added at the end
 };
 
diff --git a/Engine/Plugins/Messaging/UdpMessaging/Source/UdpMessaging/Private/Transport/UdpDeserializedMessage.cpp b/Engine/Plugins/Messaging/UdpMessaging/Source/UdpMessaging/Private/Transport/UdpDeserializedMessage.cpp
--- a/Engine/Plugins/Messaging/UdpMessaging/Source/UdpMessaging/Private/Transport/UdpDeserializedMessage.cpp
+++ b/Engine/Plugins/Messaging/UdpMessaging/Source/UdpMessaging/Private/Transport/UdpDeserializedMessage.cpp
@@ -41,7 +41,7 @@ class FUdpDeserializedMessageDetails
 public:
 	static bool DeserializeV10(FUdpDeserializedMessage& DeserializedMessage, FMemoryReader& MessageReader);
 	static bool DeserializeV11_15(FUdpDeserializedMessage& DeserializedMessage, FMemoryReader& MessageReader, bool bIsLWCBackwardCompatibilityMode);
-	static bool DeserializeV16(FUdpDeserializedMessage& DeserializedMessage, FMemoryReader& MessageReader);
+	static bool DeserializeV16_17(FUdpDeserializedMessage& DeserializedMessage, FMemoryReader& MessageReader);
 	static bool Deserialize(FUdpDeserializedMessage& DeserializedMessage, const FUdpReassembledMessage& ReassembledMessage);
 };
 
@@ -360,7 +360,7 @@ bool FUdpDeserializedMessageDetails::DeserializeV11_15(FUdpDeserializedMessage&
 	}
 }
 
-bool FUdpDeserializedMessageDetails::DeserializeV16(FUdpDeserializedMessage& DeserializedMessage, FMemoryReader& MessageReader)
+bool FUdpDeserializedMessageDetails::DeserializeV16_17(FUdpDeserializedMessage& DeserializedMessage, FMemoryReader& MessageReader)
 {
 	// message type info
 	{
@@ -524,8 +524,10 @@ bool FUdpDeserializedMessageDetails::Deserialize(FUdpDeserializedMessage& Deseri
 		break;
 	}
 	case 16:
+		// fallthrough
+	case 17:
 	{
-		return DeserializeV16(DeserializedMessage, MessageReader);
+		return DeserializeV16_17(DeserializedMessage, MessageReader);
 		break;
 	}
 
diff --git a/Engine/Plugins/Messaging/UdpMessaging/Source/UdpMessaging/Private/Transport/UdpMessageProcessor.cpp b/Engine/Plugins/Messaging/UdpMessaging/Source/UdpMessaging/Private/Transport/UdpMessageProcessor.cpp
--- a/Engine/Plugins/Messaging/UdpMessaging/Source/UdpMessaging/Private/Transport/UdpMessageProcessor.cpp
+++ b/Engine/Plugins/Messaging/UdpMessaging/Source/UdpMessaging/Private/Transport/UdpMessageProcessor.cpp
@@ -333,9 +333,10 @@ bool FUdpMessageProcessor::Init()
 		Beacon = MakeUnique<FUdpMessageBeacon>(Socket, LocalNodeId, MulticastEndpoint);
 		SocketSender = MakeUnique<FUdpSocketSender>(Socket, TEXT("FUdpMessageProcessor.Sender"));
 
-		// Current protocol version 16
+		// Current protocol version 17
 		SupportedProtocolVersions.Add(UDP_MESSAGING_TRANSPORT_PROTOCOL_VERSION);
-		// Support Protocol version 10, 11, 12, 13, 14, 15
+		// Support Protocol version 10, 11, 12, 13, 14, 15, 16
+		SupportedProtocolVersions.Add(16);
 		SupportedProtocolVersions.Add(15);
 		SupportedProtocolVersions.Add(14);
 		SupportedProtocolVersions.Add(13);
@@ -494,6 +495,7 @@ void FUdpMessageProcessor::ConsumeInboundSegments()
 				NodeInfo.NodeId = Header.SenderNodeId;
 				NodeInfo.ProtocolVersion = Header.ProtocolVersion;
 				NodeDiscoveredDelegate.ExecuteIfBound(NodeInfo.NodeId);
+				bAddedNewKnownNodes = true;
 			}
 
 			NodeInfo.ProtocolVersion = Header.ProtocolVersion;
@@ -542,6 +544,10 @@ void FUdpMessageProcessor::ConsumeInboundSegments()
 				ProcessTimeoutSegment(Segment, NodeInfo);
 				break;
 
+			case EUdpMessageSegments::Mesh:
+				ProcessMeshSegment(Segment, NodeInfo);
+				break;
+
 			default:
 				ProcessUnknownSegment(Segment, NodeInfo, (uint8)Header.SegmentType);
 			}
@@ -845,7 +851,6 @@ void FUdpMessageProcessor::ProcessPongSegment(FInboundSegment& Segment, FNodeInf
 	}
 }
 
-
 void FUdpMessageProcessor::ProcessRetransmitSegment(FInboundSegment& Segment, FNodeInfo& NodeInfo)
 {
 	FUdpMessageSegment::FRetransmitChunk RetransmitChunk;
@@ -879,13 +884,11 @@ void FUdpMessageProcessor::ProcessTimeoutSegment(FInboundSegment& Segment, FNode
 	}
 }
 
-
 void FUdpMessageProcessor::ProcessUnknownSegment(FInboundSegment& Segment, FNodeInfo& EndpointInfo, uint8 SegmentType)
 {
 	UE_LOG(LogUdpMessaging, Verbose, TEXT("Received unknown segment type '%i' from %s"), SegmentType, *Segment.Sender.ToText().ToString());
 }
 
-
 void FUdpMessageProcessor::DeliverMessage(const TSharedPtr<FUdpReassembledMessage, ESPMode::ThreadSafe>& ReassembledMessage, FNodeInfo& NodeInfo)
 {
 	// Do not deliver message while saving or garbage collecting since those deliveries will fail anyway...
@@ -961,6 +964,86 @@ void FUdpMessageProcessor::HandleSocketError(const FNodeInfo& NodeInfo) const
 	ErrorSendingToEndpointDelegate.Execute(NodeInfo.NodeId, NodeInfo.Endpoint);
 }
 
+bool FUdpMessageProcessor::CanSendKnownNodesToKnownNodes() const
+{
+	return bShareKnownNodes && bAddedNewKnownNodes;
+}
+
+void FUdpMessageProcessor::SendKnownNodesToKnownNodes()
+{
+	FUdpMessageSegment::FHeader Header;
+	{
+		Header.SenderNodeId = LocalNodeId;
+		Header.ProtocolVersion = UDP_MESSAGING_TRANSPORT_PROTOCOL_VERSION;
+		Header.SegmentType = EUdpMessageSegments::Mesh;
+	};
+
+	TArray<FIPv4Endpoint> Endpoints, OutEndpoints;
+	TArray<FGuid> EndpointIds, OutIds;
+	for (const auto& NodePair : KnownNodes)
+	{
+		EndpointIds.Add(NodePair.Key);
+		Endpoints.Add(NodePair.Value.Endpoint);
+	}
+
+	// We use TArrayView to slice parts of the array for sending in parts (in case we exceed segment sending limits).
+	TArrayView<FIPv4Endpoint> KnownEndpointsView(Endpoints);
+	TArrayView<FGuid> KnownIdView(EndpointIds);
+
+	// We can send ~80 addresses in one segment. Warn the user when we hit this value. However, we still handle meshing large networks.
+	const int32 NumEndpointsCanSend = (UDP_MESSAGING_SEGMENT_SIZE - sizeof(FUdpMessageSegment::FHeader)) / (sizeof(FIPv4Endpoint) + sizeof(FGuid)) - 1;
+	if (KnownEndpointsView.Num() > NumEndpointsCanSend)
+	{
+		UE_LOG(LogUdpMessaging, Warning, TEXT("FUdpMessageProcessor::SendKnownNodesToKnownNodes large number of endpoints to share for meshing udp transport."));
+	}
+
+	int32 Index = 0;
+	while (Index < KnownEndpointsView.Num())
+	{
+		const uint32 NextBlockIndex = FMath::Min<uint32>(Index + NumEndpointsCanSend, KnownEndpointsView.Num());
+		TArrayView<FIPv4Endpoint> SlicedEndpointView = KnownEndpointsView.Slice(Index, NextBlockIndex);
+		TArrayView<FGuid>	      SlicedIdView = KnownIdView.Slice(Index, NextBlockIndex);
+
+		OutEndpoints = SlicedEndpointView;
+		OutIds = SlicedIdView;
+		FArrayWriter Writer;
+		{
+			Writer << Header;
+			Writer << OutEndpoints;
+			Writer << OutIds;
+		}
+		UE_LOG(LogUdpMessaging, VeryVerbose, TEXT("FUdpMessageProcessor::SendKnownNodesToKnownNodes Sending updated known nodes."));
+
+		for (const FIPv4Endpoint& Endpoint : Endpoints)
+		{
+			int32 Sent;
+			if (!Socket->SendTo(Writer.GetData(), Writer.Num(), Sent, *Endpoint.ToInternetAddr()))
+			{
+				UE_LOG(LogUdpMessaging, Warning, TEXT("FUdpMessageProcessor::SendKnownNodesToKnownNodes failed to share endpoint information to %s."), *Endpoint.ToString());
+			}
+		}
+		Index += NumEndpointsCanSend;
+	}
+	bAddedNewKnownNodes = false;
+}
+
+void FUdpMessageProcessor::ProcessMeshSegment(FInboundSegment& Segment, FNodeInfo& NodeInfo)
+{
+	TArray<FGuid> Ids;
+	TArray<FIPv4Endpoint> Endpoints;
+	*Segment.Data << Endpoints;
+	*Segment.Data << Ids;
+
+	check(Ids.Num() == Endpoints.Num());
+	for (int32 Index = 0; Index < Ids.Num(); Index ++)
+	{
+		if (LocalNodeId != Ids[Index] && !KnownNodes.Find(Ids[Index]))
+		{
+			AddStaticEndpoint(Endpoints[Index]);
+		}
+	}
+}
+
 void FUdpMessageProcessor::UpdateKnownNodes()
 {
 	SCOPED_MESSAGING_TRACE(FUdpMessageProcessor_UpdateKnownNodes);
@@ -970,6 +1053,11 @@ void FUdpMessageProcessor::UpdateKnownNodes()
 	UpdateNodesPerVersion();
 	Beacon->SetEndpointCount(KnownNodes.Num() + 1);
 
+	if (CanSendKnownNodesToKnownNodes())
+	{
+		SendKnownNodesToKnownNodes();
+	}
+
 	bool bSuccess = true;
 	for (auto& KnownNodePair : KnownNodes)
 	{
diff --git a/Engine/Plugins/Messaging/UdpMessaging/Source/UdpMessaging/Private/Transport/UdpMessageProcessor.h b/Engine/Plugins/Messaging/UdpMessaging/Source/UdpMessaging/Private/Transport/UdpMessageProcessor.h
--- a/Engine/Plugins/Messaging/UdpMessaging/Source/UdpMessaging/Private/Transport/UdpMessageProcessor.h
+++ b/Engine/Plugins/Messaging/UdpMessaging/Source/UdpMessaging/Private/Transport/UdpMessageProcessor.h
@@ -570,6 +570,10 @@ public:
 	 */
 	FMessageTransportStatistics GetStats(FGuid Node) const;
 
+	void SetShareKnownNodesState(bool bInShareKnownNodes)
+	{
+		bShareKnownNodes = bInShareKnownNodes;
+	}
 public:
 
 	// @todo gmp: remove the need for this typedef
@@ -772,6 +776,14 @@ protected:
 	 */
 	void ProcessTimeoutSegment(FInboundSegment& Segment, FNodeInfo& NodeInfo);
 
+	/**
+	 * Processes a Mesh segment.
+	 *
+	 * @param Segment The segment to process.
+	 * @param NodeInfo Details for the node that sent the segment.
+	 */
+	void ProcessMeshSegment(FInboundSegment& Segment, FNodeInfo& NodeInfo);
+
 	/**
 	 * Processes an unknown segment type.
 	 *
@@ -826,6 +838,13 @@ protected:
 	virtual void Tick() override;
 
 private:
+
+	/** Send our list of known nodes to the known nodes. */
+	void SendKnownNodesToKnownNodes();
+
+	/** Do we support sending known nodes to our known nodes. */
+	bool CanSendKnownNodesToKnownNodes() const;
+
 	/** Start the message processor thread. */
 	void StartThread();
 
@@ -874,6 +893,9 @@ private:
 	/** Holds the collection of known remote nodes. */
 	TMap<FGuid, FNodeInfo> KnownNodes;
 
+	/** Indicate if we added new nodes during this frame. */
+	bool bAddedNewKnownNodes = false;
+
 	/** Holds the local node identifier. */
 	FGuid LocalNodeId;
 
@@ -889,6 +911,9 @@ private:
 	/** Holds the socket sender.*/
 	TUniquePtr<FUdpSocketSender> SocketSender;
 
+	/** Enable sharing of our known nodes. */
+	bool bShareKnownNodes = false;
+
 	/** Holds a flag indicating that the thread is stopping. */
 	bool bStopping;
 
diff --git a/Engine/Plugins/Messaging/UdpMessaging/Source/UdpMessaging/Private/Transport/UdpMessageTransport.cpp b/Engine/Plugins/Messaging/UdpMessaging/Source/UdpMessaging/Private/Transport/UdpMessageTransport.cpp
--- a/Engine/Plugins/Messaging/UdpMessaging/Source/UdpMessaging/Private/Transport/UdpMessageTransport.cpp
+++ b/Engine/Plugins/Messaging/UdpMessaging/Source/UdpMessaging/Private/Transport/UdpMessageTransport.cpp
@@ -326,6 +326,8 @@ bool FUdpMessageTransport::StartTransport(IMessageTransportHandler& Handler)
 		MessageProcessor->OnError().BindRaw(this, &FUdpMessageTransport::HandleProcessorError);
 		MessageProcessor->OnErrorSendingToEndpoint_UdpMessageProcessorThread().BindRaw(this, &FUdpMessageTransport::HandleEndpointCommunicationError);
 		MessageProcessor->OnCanAcceptEndpoint_UdpMessageProcessorThread().BindRaw(this, &FUdpMessageTransport::HandleProcessorEndpointCheck);
+		UUdpMessagingSettings const * CurrentSettings = GetDefault<UUdpMessagingSettings>();
+		MessageProcessor->SetShareKnownNodesState(CurrentSettings->bShareKnownNodesWithActiveConnections);
 	}
 
 	if (MulticastSocket != nullptr)
@@ -478,7 +480,7 @@ void FUdpMessageTransport::HandleProcessorError()
 			// won't be able to access the settings CDO even if the transport still exists.
 			if (!UObjectInitialized())
 			{
-					return;
+				return;
 			}
 
 			if (TSharedPtr<FUdpMessageTransport, ESPMode::ThreadSafe> Transport = WeakTransport.Pin())
diff --git a/Engine/Plugins/Messaging/UdpMessaging/Source/UdpMessaging/Private/Transport/UdpSerializeMessageTask.cpp b/Engine/Plugins/Messaging/UdpMessaging/Source/UdpMessaging/Private/Transport/UdpSerializeMessageTask.cpp
--- a/Engine/Plugins/Messaging/UdpMessaging/Source/UdpMessaging/Private/Transport/UdpSerializeMessageTask.cpp
+++ b/Engine/Plugins/Messaging/UdpMessaging/Source/UdpMessaging/Private/Transport/UdpSerializeMessageTask.cpp
@@ -94,7 +94,7 @@ void SerializeMessageV11_15(FArchive& Archive, const TSharedRef<IMessageContext,
 }
 
 /** Serialization Routine for message using Protocol version 16. */
-void SerializeMessageV16(FArchive& Archive, const TSharedRef<IMessageContext, ESPMode::ThreadSafe>& MessageContext, EUdpMessageFormat MessageFormat, const EStructSerializerBackendFlags StructSerializerBackendFlags)
+void SerializeMessageV16_17(FArchive& Archive, const TSharedRef<IMessageContext, ESPMode::ThreadSafe>& MessageContext, EUdpMessageFormat MessageFormat, const EStructSerializerBackendFlags StructSerializerBackendFlags)
 {
 	const FTopLevelAssetPath& MessageType = MessageContext->GetMessageTypePathName();
 	Archive << const_cast<FTopLevelAssetPath&>(MessageType);
@@ -189,7 +189,12 @@ void FUdpSerializeMessageTask::DoTask(ENamedThreads::Type CurrentThread, const F
 
 			case 16:
 				ProtocolMaxSegmentSize = UDP_MESSAGING_SEGMENT_SIZE * (int64)INT32_MAX;
-				UdpSerializeMessageTaskDetails::SerializeMessageV16(Archive, MessageContext, SerializedMessage->GetFormat(), EStructSerializerBackendFlags::Default);
+				UdpSerializeMessageTaskDetails::SerializeMessageV16_17(Archive, MessageContext, SerializedMessage->GetFormat(), EStructSerializerBackendFlags::Default);
+				break;
+
+			case 17:
+				ProtocolMaxSegmentSize = UDP_MESSAGING_SEGMENT_SIZE * (int64)INT32_MAX;
+				UdpSerializeMessageTaskDetails::SerializeMessageV16_17(Archive, MessageContext, SerializedMessage->GetFormat(), EStructSerializerBackendFlags::Default);
 				break;
 
 			default:
diff --git a/Engine/Plugins/Messaging/UdpMessaging/Source/UdpMessaging/Private/UdpMessagingModule.cpp b/Engine/Plugins/Messaging/UdpMessaging/Source/UdpMessaging/Private/UdpMessagingModule.cpp
--- a/Engine/Plugins/Messaging/UdpMessaging/Source/UdpMessaging/Private/UdpMessagingModule.cpp
+++ b/Engine/Plugins/Messaging/UdpMessaging/Source/UdpMessaging/Private/UdpMessagingModule.cpp
@@ -621,6 +621,7 @@ protected:
 		{
 			// Parse value overrides (if present)
 			FParse::Bool(CommandLine, TEXT("-UDPMESSAGING_TRANSPORT_ENABLE="), Settings->EnableTransport);
+			FParse::Bool(CommandLine, TEXT("-UDPMESSAGING_SHARE_KNOWN_NODES="), Settings->bShareKnownNodesWithActiveConnections);
 			FParse::Value(CommandLine, TEXT("-UDPMESSAGING_TRANSPORT_UNICAST="), Settings->UnicastEndpoint);
 			FParse::Value(CommandLine, TEXT("-UDPMESSAGING_TRANSPORT_MULTICAST="), Settings->MulticastEndpoint);
 			FParse::Value(CommandLine, TEXT("-UDPMESSAGING_WORK_QUEUE_SIZE="), Settings->WorkQueueSize);
@@ -633,6 +634,11 @@ protected:
 			{
 				Settings->StaticEndpoints.AddUnique(CmdStaticEndpoint);
 			}
+
+			if (Settings->bShareKnownNodesWithActiveConnections)
+			{
+				UE_LOG(LogUdpMessaging, Display, TEXT("Sharing known nodes with endpoints is enabled."));
+			}
 		}
 	}
 
diff --git a/Engine/Plugins/Messaging/UdpMessaging/Source/UdpMessaging/Private/UdpMessagingPrivate.h b/Engine/Plugins/Messaging/UdpMessaging/Source/UdpMessaging/Private/UdpMessagingPrivate.h
--- a/Engine/Plugins/Messaging/UdpMessaging/Source/UdpMessaging/Private/UdpMessagingPrivate.h
+++ b/Engine/Plugins/Messaging/UdpMessaging/Source/UdpMessaging/Private/UdpMessagingPrivate.h
@@ -28,4 +28,4 @@ DECLARE_LOG_CATEGORY_EXTERN(LogUdpMessaging, Log, All);
  * Defines the protocol version of the UDP message transport.
  * @note When changing the version, ensure to update the serialization/deserialization code in UdpSerializeMessageTask.cpp/UdpDeserializedMessage.cpp and the supported versions in FUdpMessageProcessor::Init().
  */
-#define UDP_MESSAGING_TRANSPORT_PROTOCOL_VERSION 16
+#define UDP_MESSAGING_TRANSPORT_PROTOCOL_VERSION 17
diff --git a/Engine/Plugins/Messaging/UdpMessaging/Source/UdpMessaging/Public/Shared/UdpMessagingSettings.h b/Engine/Plugins/Messaging/UdpMessaging/Source/UdpMessaging/Public/Shared/UdpMessagingSettings.h
--- a/Engine/Plugins/Messaging/UdpMessaging/Source/UdpMessaging/Public/Shared/UdpMessagingSettings.h
+++ b/Engine/Plugins/Messaging/UdpMessaging/Source/UdpMessaging/Public/Shared/UdpMessagingSettings.h
@@ -141,6 +141,12 @@ public:
 	UPROPERTY(config, EditAnywhere, Category=Transport, AdvancedDisplay)
 	TArray<FString> ExcludedEndpoints;
 
+	/**
+	 * Whether to share our list of known endpoints with those endpoints.  This allows message processors to completely mesh their networks.
+	 */
+	UPROPERTY(config, EditAnywhere, Category=Transport, AdvancedDisplay)
+	bool bShareKnownNodesWithActiveConnections = false;
+
 public:
 
 	/** Whether the UDP tunnel is enabled. */
diff --git a/Engine/Plugins/VirtualProduction/Switchboard/Source/Switchboard/switchboard/switchboard_application.py b/Engine/Plugins/VirtualProduction/Switchboard/Source/Switchboard/switchboard/switchboard_application.py
--- a/Engine/Plugins/VirtualProduction/Switchboard/Source/Switchboard/switchboard/switchboard_application.py
+++ b/Engine/Plugins/VirtualProduction/Switchboard/Source/Switchboard/switchboard/switchboard_application.py
@@ -198,6 +198,7 @@ class MultiUserApplication:
             args.extend([
                 f'-CONCERTSERVER="{CONFIG.MUSERVER_SERVER_NAME.get_value()}"',
                 f'{CONFIG.MUSERVER_COMMAND_LINE_ARGUMENTS.get_value()}',
+                f'-UDPMESSAGING_SHARE_KNOWN_NODES=1',
                 f'{self.get_mu_server_endpoint_arg()}',
                 f'{self.get_mu_server_multicast_arg()}'])
 
