- Fixing crash in UDP messaging meshing when too many endpoints are found

diff --git a/Engine/Plugins/Messaging/UdpMessaging/Source/UdpMessaging/Private/Transport/UdpMessageProcessor.cpp b/Engine/Plugins/Messaging/UdpMessaging/Source/UdpMessaging/Private/Transport/UdpMessageProcessor.cpp
--- a/Engine/Plugins/Messaging/UdpMessaging/Source/UdpMessaging/Private/Transport/UdpMessageProcessor.cpp
+++ b/Engine/Plugins/Messaging/UdpMessaging/Source/UdpMessaging/Private/Transport/UdpMessageProcessor.cpp
@@ -994,15 +994,15 @@ void FUdpMessageProcessor::SendKnownNodesToKnownNodes()
 	const int32 NumEndpointsCanSend = (UDP_MESSAGING_SEGMENT_SIZE - sizeof(FUdpMessageSegment::FHeader)) / (sizeof(FIPv4Endpoint) + sizeof(FGuid)) - 1;
 	if (KnownEndpointsView.Num() > NumEndpointsCanSend)
 	{
-		UE_LOG(LogUdpMessaging, Warning, TEXT("FUdpMessageProcessor::SendKnownNodesToKnownNodes large number of endpoints to share for meshing udp transport."));
+		UE_LOG(LogUdpMessaging, Warning, TEXT("FUdpMessageProcessor::SendKnownNodesToKnownNodes large number of endpoints (%d) to share for meshing udp transport."), KnownEndpointsView.Num());
 	}
 
 	int32 Index = 0;
 	while (Index < KnownEndpointsView.Num())
 	{
-		const uint32 NextBlockIndex = FMath::Min<uint32>(Index + NumEndpointsCanSend, KnownEndpointsView.Num());
-		TArrayView<FIPv4Endpoint> SlicedEndpointView = KnownEndpointsView.Slice(Index, NextBlockIndex);
-		TArrayView<FGuid>	      SlicedIdView = KnownIdView.Slice(Index, NextBlockIndex);
+		const uint32 NextBlockItemCount = FMath::Min<uint32>(Index + NumEndpointsCanSend, KnownEndpointsView.Num() - Index);
+		TArrayView<FIPv4Endpoint> SlicedEndpointView = KnownEndpointsView.Slice(Index, NextBlockItemCount);
+		TArrayView<FGuid>	      SlicedIdView = KnownIdView.Slice(Index, NextBlockItemCount);
 
 		OutEndpoints = SlicedEndpointView;
 		OutIds = SlicedIdView;
