Improve detection of bad ip addresses.  Also allow excluded endpoints to include wildcard style addresses.

diff --git a/Engine/Plugins/Messaging/UdpMessaging/Source/UdpMessaging/Private/Tests/UdpMessageTransportTest.cpp b/Engine/Plugins/Messaging/UdpMessaging/Source/UdpMessaging/Private/Tests/UdpMessageTransportTest.cpp
--- a/Engine/Plugins/Messaging/UdpMessaging/Source/UdpMessaging/Private/Tests/UdpMessageTransportTest.cpp
+++ b/Engine/Plugins/Messaging/UdpMessaging/Source/UdpMessaging/Private/Tests/UdpMessageTransportTest.cpp
@@ -11,7 +11,6 @@
 
 IMPLEMENT_SIMPLE_AUTOMATION_TEST(FUdpMessageTransportTest, "System.Core.Messaging.Transports.Udp.UdpMessageTransport", EAutomationTestFlags::ApplicationContextMask | EAutomationTestFlags::EngineFilter)
 
-
 class FUdpMessageTransportTestState
 	: public IMessageTransportHandler
 {
@@ -21,7 +20,7 @@ public:
 		: NumReceivedMessages(0)
 	{
 		TArray<FIPv4Endpoint> StaticEndpoints;
-		TArray<FIPv4Endpoint> ExcludeEndpoints;
+		TArray<FString> ExcludeEndpoints;
 		Transport = MakeShared<FUdpMessageTransport, ESPMode::ThreadSafe>(UnicastEndpoint, MulticastEndpoint,
 																		  MoveTemp(StaticEndpoints),
 																		  MoveTemp(ExcludeEndpoints), MulticastTimeToLive);
@@ -88,6 +87,8 @@ public:
 
 private:
 
+	void TestWildcardEndpointExclusion();
+
 	TArray<FGuid> DiscoveredNodes;
 	TArray<FGuid> LostNodes;
 	int32 NumReceivedMessages;
@@ -96,6 +97,39 @@ private:
 	TSharedPtr<IMessageContext, ESPMode::ThreadSafe> LastMessageContext;
 };
 
+IMPLEMENT_SIMPLE_AUTOMATION_TEST(FUdpMessageTransportTestExclusionList, "System.Core.Messaging.Transports.Udp.UdpMessageTransportExclusionList", EAutomationTestFlags::ApplicationContextMask | EAutomationTestFlags::EngineFilter)
+
+bool FUdpMessageTransportTestExclusionList::RunTest(const FString& Parameters)
+{
+	const FIPv4Endpoint SimpleEndpoint(FIPv4Address(192, 168, 10, 120), 12777);
+
+	TArray<FString> PassTests = {
+		TEXT("192.16?.10.*"),
+		TEXT("192.168.10.120"),
+		TEXT("192.168.10.*:0"),
+		TEXT("192.168.10.120:12777"),
+		TEXT("*.*.*.*:12777")
+	};
+	for (const FString& Test : PassTests)
+	{
+		AddInfo(FString::Printf(TEXT("Testing wildcard string %s against target %s"), *Test, *SimpleEndpoint.ToString()));
+		TestTrue(TEXT("Wilcard match should succeed."),	 FUdpMessageTransport::MatchesEndpoint(Test, SimpleEndpoint));
+	}
+
+	TArray<FString> FailTests = {
+		TEXT("192.168.10.11"),
+		TEXT("192.168.10.120:1210"),
+		TEXT("192.168.*.55:0"),
+		TEXT("192.168.11.11")
+    };
+
+	for (const FString& Test : FailTests)
+	{
+		AddInfo(FString::Printf(TEXT("Testing wildcard string %s against target %s"), *Test, *SimpleEndpoint.ToString()));
+		TestFalse(TEXT("Wilcard match should not succeed."), FUdpMessageTransport::MatchesEndpoint(Test, SimpleEndpoint));
+	}
+	return true;
+}
 
 bool FUdpMessageTransportTest::RunTest(const FString& Parameters)
 {
diff --git a/Engine/Plugins/Messaging/UdpMessaging/Source/UdpMessaging/Private/Transport/UdpMessageTransport.cpp b/Engine/Plugins/Messaging/UdpMessaging/Source/UdpMessaging/Private/Transport/UdpMessageTransport.cpp
--- a/Engine/Plugins/Messaging/UdpMessaging/Source/UdpMessaging/Private/Transport/UdpMessageTransport.cpp
+++ b/Engine/Plugins/Messaging/UdpMessaging/Source/UdpMessaging/Private/Transport/UdpMessageTransport.cpp
@@ -101,7 +101,7 @@ struct FDenyList
 static FDenyList GlobalDenyCandidateTable;
 }
 
-FUdpMessageTransport::FUdpMessageTransport(const FIPv4Endpoint& InUnicastEndpoint, const FIPv4Endpoint& InMulticastEndpoint, TArray<FIPv4Endpoint> InStaticEndpoints, TArray<FIPv4Endpoint> InExcludedEndpoints, uint8 InMulticastTtl)
+FUdpMessageTransport::FUdpMessageTransport(const FIPv4Endpoint& InUnicastEndpoint, const FIPv4Endpoint& InMulticastEndpoint, TArray<FIPv4Endpoint> InStaticEndpoints, TArray<FString> InExcludedEndpointsAsString, uint8 InMulticastTtl)
 	: MessageProcessor(nullptr)
 	, MulticastEndpoint(InMulticastEndpoint)
 	, MulticastReceiver(nullptr)
@@ -115,7 +115,7 @@ FUdpMessageTransport::FUdpMessageTransport(const FIPv4Endpoint& InUnicastEndpoin
 	, UnicastSocket(nullptr)
 #endif
 	, StaticEndpoints(MoveTemp(InStaticEndpoints))
-	, ExcludedEndpoints(MoveTemp(InExcludedEndpoints))
+	, ExcludedEndpoints(MoveTemp(InExcludedEndpointsAsString))
 {
 }
 
@@ -441,7 +441,95 @@ void FUdpMessageTransport::HandleProcessorNodeLost(const FGuid& LostNodeId)
 	TransportHandler->ForgetTransportNode(LostNodeId);
 }
 
-bool FUdpMessageTransport::HandleProcessorEndpointCheck(const FGuid& EndpointNodeId, const FIPv4Endpoint& SenderIpAddress)
+namespace UE::UdpMessageTransport::Private
+{
+	TArray<FString> TokenizeStringAddress(const FString &Address)
+	{
+		TArray<FString> Tokens;
+		bool bIsValid = Address.ParseIntoArray(Tokens, TEXT("."), false /*CullEmpty*/) == 4;
+		if (bIsValid)
+		{
+			for (const FString& Token : Tokens)
+			{
+				if (!FCString::IsNumeric(*Token))
+				{
+					// If it's not a pure number, look for individual characters
+					for (int32 Index = 0; Index < Token.Len(); ++Index)
+					{
+						const TCHAR& Character = Token[Index];
+						if (!(FChar::IsDigit(Character) || Character == '*' || Character == '?'))
+						{
+							bIsValid = false;
+							break;
+						}
+					}
+
+					if (!bIsValid)
+					{
+						break;
+					}
+				}
+			}
+		}
+		return Tokens;
+	}
+
+	bool MatchesIPv4Address(const FString& SourceAddress, const FIPv4Address& TargetAddress)
+	{
+		FString TargetAsString = TargetAddress.ToString();
+		TArray<FString> SourceTokenized = TokenizeStringAddress(SourceAddress);
+		TArray<FString> TargetTokenized =  TokenizeStringAddress(TargetAsString);
+		if (SourceTokenized.Num() == 4 && SourceTokenized.Num() == TargetTokenized.Num())
+		{
+			for (int32 TokenIndex = 0; TokenIndex < 4; ++TokenIndex)
+			{
+				const FString& Source = SourceTokenized[TokenIndex];
+				const FString& Target = TargetTokenized[TokenIndex];
+				if (!Target.MatchesWildcard(Source))
+				{
+					return false;
+				}
+			}
+
+			// It must have passed all tokens.
+			return true;
+		}
+		return true;
+	}
+
+	int32 GetPortNumber(const FString& InPortAsString)
+	{
+		if (InPortAsString.IsNumeric())
+		{
+			return FCString::Atoi(*InPortAsString);
+		}
+		return -1;
+	}
+
+
+}
+
+bool FUdpMessageTransport::MatchesEndpoint(const FString& SourceEndpoint, const FIPv4Endpoint& TargetEndpoint)
+{
+	using namespace UE::UdpMessageTransport::Private;
+	if (SourceEndpoint.IsEmpty())
+	{
+		return false;
+	}
+
+	TArray<FString> IpPortTokens;
+	const int32 Items = SourceEndpoint.ParseIntoArray(IpPortTokens, TEXT(":"));
+	if (Items <= 0)
+	{
+		return false;
+	}
+
+	const int32 Port = Items == 1 ? 0 : GetPortNumber(IpPortTokens[1]);
+	return (Port == 0 || Port == TargetEndpoint.Port)
+		&& MatchesIPv4Address(IpPortTokens[0], TargetEndpoint.Address);
+}
+
+bool FUdpMessageTransport::HandleProcessorEndpointCheck(const FGuid& EndpointNodeId, const FIPv4Endpoint& SenderEndpoint)
 {
 	if (CVarEndpointDenyListEnabled.GetValueOnAnyThread())
 	{
@@ -451,10 +539,9 @@ bool FUdpMessageTransport::HandleProcessorEndpointCheck(const FGuid& EndpointNod
 		}
 	}
 
-	for (const FIPv4Endpoint& Endpoint : ExcludedEndpoints)
+	for (const FString& Endpoint : ExcludedEndpoints)
 	{
-		if ((Endpoint.Port == 0 && SenderIpAddress.Address == Endpoint.Address)
-			|| SenderIpAddress == Endpoint)
+		if (MatchesEndpoint(Endpoint, SenderEndpoint))
 		{
 			return false;
 		}
diff --git a/Engine/Plugins/Messaging/UdpMessaging/Source/UdpMessaging/Private/Transport/UdpMessageTransport.h b/Engine/Plugins/Messaging/UdpMessaging/Source/UdpMessaging/Private/Transport/UdpMessageTransport.h
--- a/Engine/Plugins/Messaging/UdpMessaging/Source/UdpMessaging/Private/Transport/UdpMessageTransport.h
+++ b/Engine/Plugins/Messaging/UdpMessaging/Source/UdpMessaging/Private/Transport/UdpMessageTransport.h
@@ -44,13 +44,13 @@ public:
 	 * @param InLocalEndpoint The local IP endpoint to receive messages on.
 	 * @param InMulticastEndpoint The multicast group endpoint to transport messages to.
 	 * @param InStaticEndpoints array of endpoints to include in communication
-	 * @param InExcludedEndpoints array of endpoints to exclude in communication
+	 * @param InExcludedEndpoints array of string formatted endpoints to exclude in communication.
 	 * @param InMulticastTtl The multicast time-to-live.
 	 */
 	FUdpMessageTransport(const FIPv4Endpoint& InLocalEndpoint,
 						 const FIPv4Endpoint& InMulticastEndpoint,
 						 TArray<FIPv4Endpoint> InStaticEndpoints,
-						 TArray<FIPv4Endpoint> InExcludedEndpoints,
+						 TArray<FString> InExcludedEndpointsAsString,
 						 uint8 InMulticastTtl);
 
 	/** Virtual destructor. */
@@ -59,6 +59,12 @@ public:
 	/** Notifies the transport that the application is about to exit. */
 	void OnAppPreExit();
 
+	/**
+	 *  Helper method to detect if a given target endpoint matches the wildcard deliminted source endpoint. The source endpoint should be in
+	 *  IPv4 format.  A.B.C.D:<OptionalPortNumber>.  Wildcard characters are allowed.
+	 */
+	static bool MatchesEndpoint(const FString& SourceEndpoint, const FIPv4Endpoint& TargetEndpoint);
+
 public:
 	/**
 	 * Add a static endpoint to the transport
@@ -174,6 +180,6 @@ private:
 	/** Holds the static endpoints. */
 	TSet<FIPv4Endpoint> StaticEndpoints;
 
-	/** Excluded list of IP addresses not allowed to talk to UDP processor */
-	TArray<FIPv4Endpoint> ExcludedEndpoints;
+	/** Excluded list of IP addresses as string value (can contain wildcards) not allowed to talk to UDP processor */
+	TArray<FString> ExcludedEndpoints;
 };
diff --git a/Engine/Plugins/Messaging/UdpMessaging/Source/UdpMessaging/Private/UdpMessagingModule.cpp b/Engine/Plugins/Messaging/UdpMessaging/Source/UdpMessaging/Private/UdpMessagingModule.cpp
--- a/Engine/Plugins/Messaging/UdpMessaging/Source/UdpMessaging/Private/UdpMessagingModule.cpp
+++ b/Engine/Plugins/Messaging/UdpMessaging/Source/UdpMessaging/Private/UdpMessagingModule.cpp
@@ -467,6 +467,14 @@ protected:
 		{
 			bParsedAddr = FIPv4Endpoint::FromHostAndPort(InEndpointString, OutEndpoint);
 		}
+		if (bParsedAddr)
+		{
+			// Detect 169.254.x.x addresses.
+			if (OutEndpoint.Address.A == 169 &&	OutEndpoint.Address.B == 254)
+			{
+				UE_LOG(LogUdpMessaging, Warning, TEXT("Detected IPv4 address in the form of 169.254.x.x. This is a link assigned address and may prevent you from reach external endpoints. "));
+			}
+		}
 		return bParsedAddr;
 	}
 
@@ -528,9 +536,6 @@ protected:
 		TArray<FIPv4Endpoint> StaticEndpoints = AdditionalStaticEndpoints.Array();
 		StaticEndpoints += ParseStringArrayAddresses(Settings->StaticEndpoints);
 
-		// Addresses to deny on transport.
-		TArray<FIPv4Endpoint> ExcludedEndpoints = ParseStringArrayAddresses(Settings->ExcludedEndpoints);
-
 		if (Settings->MulticastTimeToLive == 0)
 		{
 			Settings->MulticastTimeToLive = 1;
@@ -544,7 +549,7 @@ protected:
 		UE_LOG(LogUdpMessaging, Log, TEXT("Initializing bridge on interface %s to multicast group %s."), *UnicastEndpoint.ToString(), *MulticastEndpoint.ToText().ToString());
 
 		TSharedRef<FUdpMessageTransport, ESPMode::ThreadSafe> Transport = MakeShared<FUdpMessageTransport, ESPMode::ThreadSafe>(
-			UnicastEndpoint, MulticastEndpoint, MoveTemp(StaticEndpoints), MoveTemp(ExcludedEndpoints), Settings->MulticastTimeToLive);
+			UnicastEndpoint, MulticastEndpoint, MoveTemp(StaticEndpoints), Settings->ExcludedEndpoints, Settings->MulticastTimeToLive);
 		WeakBridgeTransport = Transport;
 		MessageBridge = FMessageBridgeBuilder()
 			.UsingTransport(Transport);
diff --git a/Engine/Plugins/Messaging/UdpMessaging/Source/UdpMessaging/Public/Shared/UdpMessagingSettings.h b/Engine/Plugins/Messaging/UdpMessaging/Source/UdpMessaging/Public/Shared/UdpMessagingSettings.h
--- a/Engine/Plugins/Messaging/UdpMessaging/Source/UdpMessaging/Public/Shared/UdpMessagingSettings.h
+++ b/Engine/Plugins/Messaging/UdpMessaging/Source/UdpMessaging/Public/Shared/UdpMessagingSettings.h
@@ -134,7 +134,11 @@ public:
 	TArray<FString> StaticEndpoints;
 
 	/**
-	 * List of IP addresses that are banned from communicating with this client.
+	 * List of IP addresses that are banned from communicating with this client. Wildcard IPv4 addresses can be specified by using '*' and '?'
+	 * characters.
+	 *
+	 * For example, 192.168.12.* will deny all addresses on the .12 subnet for the class C network.
+	 * The string 192.16?.12.10 will deny all addresses in the range 192.160.12.10-192.169.12.10.
 	 *
 	 * The format is IP_ADDRESS:PORT_NUMBER.  If port number is 0 then all ports are blocked.
 	 */
diff --git a/Engine/Plugins/VirtualProduction/Switchboard/Source/Switchboard/switchboard/config.py b/Engine/Plugins/VirtualProduction/Switchboard/Source/Switchboard/switchboard/config.py
--- a/Engine/Plugins/VirtualProduction/Switchboard/Source/Switchboard/switchboard/config.py
+++ b/Engine/Plugins/VirtualProduction/Switchboard/Source/Switchboard/switchboard/config.py
@@ -423,6 +423,7 @@ class Setting(QtCore.QObject):
 
         if self.is_overridden(override_device_name):
             self.override_value(override_device_name, self._value)
+
             # Update UI
             self._on_widget_value_changed(self._value, override_device_name)
             self._on_setting_changed(self._value, override_device_name=override_device_name)
@@ -1838,12 +1839,73 @@ class AddressSetting(OptionSetting):
     def generate_possible_addresses(self):
         return set[str]()
 
+def is_valid_unicast_address(address):
+    # 169.254.x.x addresses can be generated by the OS when a failure to negotiate with DHCP.
+    if address.startswith("169.254"):
+        LOGGER.warning(f'Detected link assigned address {address}.  This address is not reachable outside the local link adapter. Please check your settings.')
+
+    if address == '0.0.0.0':
+        LOGGER.error('The IF_ANY address, 0.0.0.0, is not a valid switchboard unicast address.')
+        return False
+
+    split_addr = address.split('.')
+    if len(split_addr) != 4:
+        LOGGER.error(f'Unicast address ({address}) does not contain four parts.')
+        return False
+
+    for part in reversed(split_addr):
+        if not part.isdigit() or part == '255':
+            LOGGER.error(f'Broadcast-like addresses are not allowed.')
+            return False
+
+    classA_part = int(split_addr[0])
+    if classA_part >= 224 and classA_part <= 239:
+        LOGGER.error(f'Unicast address {address} is a multicast address and not valid.')
+        return False
+
+    return True
+
+def is_local_classC_ip(address):
+    if address.startswith("192."):
+        return True
+    return False
+
+def is_local_classA_ip(address):
+    if address.startswith("10."):
+        return True
+    return False
 
 class LocalAddressSetting(AddressSetting):
+    def check_and_fix_bad_ip(self):
+        orig_value = self.get_value()
+        if is_valid_unicast_address(orig_value):
+            return
+
+        addresses = self.generate_possible_addresses()
+        def iterate_over_addresses():
+            classC = filter(lambda x: is_local_classC_ip(x), addresses)
+            classA = filter(lambda x: is_local_classA_ip(x), addresses)
+            other  = filter(lambda x: not is_local_classC_ip(x) and not is_local_classA_ip(x), addresses)
+            for ip in classC:
+                yield ip
+            for ip in classA:
+                yield ip
+            for ip in other:
+                yield ip
+
+        # Pick class C first, then class A followed by other detected IP.  Finally fallback to localhost
+        new_value = "127.0.0.1"
+        for addr in iterate_over_addresses():
+            if is_valid_unicast_address(addr):
+                new_value = addr
+                break
+
+        LOGGER.error(f'{orig_value} is not a valid IP address changing default IP to {new_value}')
+        self.update_value(new_value)
+
     def generate_possible_addresses(self):
         addresses = set[str]()
-        for address in socket.getaddrinfo(socket.gethostname(), None,
-                                          socket.AF_INET):
+        for address in socket.getaddrinfo(socket.gethostname(), None, socket.AF_INET):
             addresses.add(str(address[4][0]))
         addresses.add('127.0.0.1')
         addresses.add('localhost')
@@ -2849,6 +2911,11 @@ class UserSettings(object):
         self.save()
 
     def save(self):
+        # Save will always happen after an initial load so let's validate that
+        # the IP address is correct.
+        #
+        self.ADDRESS.check_and_fix_bad_ip()
+
         data = {
             'config': '',
             'address': self.ADDRESS.get_value(),
