Fix possible DoS for UdpMessaging

diff --git a/Engine/Plugins/Messaging/UdpMessaging/Source/UdpMessaging/Private/Shared/UdpMessageSegment.h b/Engine/Plugins/Messaging/UdpMessaging/Source/UdpMessaging/Private/Shared/UdpMessageSegment.h
--- a/Engine/Plugins/Messaging/UdpMessaging/Source/UdpMessaging/Private/Shared/UdpMessageSegment.h
+++ b/Engine/Plugins/Messaging/UdpMessaging/Source/UdpMessaging/Private/Shared/UdpMessageSegment.h
@@ -357,6 +357,13 @@ namespace FUdpMessageSegment
 		{
 			Ar << MessageId << Segments;
 		}
+
+		int32 GetMessageId(FArchive& Ar)
+		{
+			int32 Id;
+			Ar << Id;
+			return Id;
+		}
 	};
 
 
diff --git a/Engine/Plugins/Messaging/UdpMessaging/Source/UdpMessaging/Private/Transport/UdpMessageProcessor.cpp b/Engine/Plugins/Messaging/UdpMessaging/Source/UdpMessaging/Private/Transport/UdpMessageProcessor.cpp
--- a/Engine/Plugins/Messaging/UdpMessaging/Source/UdpMessaging/Private/Transport/UdpMessageProcessor.cpp
+++ b/Engine/Plugins/Messaging/UdpMessaging/Source/UdpMessaging/Private/Transport/UdpMessageProcessor.cpp
@@ -854,6 +854,13 @@ void FUdpMessageProcessor::ProcessPongSegment(FInboundSegment& Segment, FNodeInf
 void FUdpMessageProcessor::ProcessRetransmitSegment(FInboundSegment& Segment, FNodeInfo& NodeInfo)
 {
 	FUdpMessageSegment::FRetransmitChunk RetransmitChunk;
+	int32 TargetMessageId = RetransmitChunk.GetMessageId(*Segment.Data);
+	if (NodeInfo.Segmenters.IsEmpty() || !NodeInfo.Segmenters.Contains(TargetMessageId))
+	{
+		// Ignore this message because we have no segmenters to retransmit.
+		return;
+	}
+
 	RetransmitChunk.Serialize(*Segment.Data, NodeInfo.ProtocolVersion);
 
 	TSharedPtr<FUdpMessageSegmenter> Segmenter = NodeInfo.Segmenters.FindRef(RetransmitChunk.MessageId);
@@ -873,6 +880,11 @@ void FUdpMessageProcessor::ProcessRetransmitSegment(FInboundSegment& Segment, FN
 
 void FUdpMessageProcessor::ProcessTimeoutSegment(FInboundSegment& Segment, FNodeInfo& NodeInfo)
 {
+	if (NodeInfo.Segmenters.IsEmpty())
+	{
+		// Ignore this message because we have no segmenters to retransmit.
+		return;
+	}
 	FUdpMessageSegment::FTimeoutChunk TimeoutChunk;
 	TimeoutChunk.Serialize(*Segment.Data, NodeInfo.ProtocolVersion);
 
diff --git a/Engine/Plugins/Messaging/UdpMessaging/Source/UdpMessaging/Private/Transport/UdpMessageTransport.cpp b/Engine/Plugins/Messaging/UdpMessaging/Source/UdpMessaging/Private/Transport/UdpMessageTransport.cpp
--- a/Engine/Plugins/Messaging/UdpMessaging/Source/UdpMessaging/Private/Transport/UdpMessageTransport.cpp
+++ b/Engine/Plugins/Messaging/UdpMessaging/Source/UdpMessaging/Private/Transport/UdpMessageTransport.cpp
@@ -266,9 +266,13 @@ bool FUdpMessageTransport::StartTransport(IMessageTransportHandler& Handler)
 	if (UnicastSocket == nullptr)
 	{
 		UE_LOG(LogUdpMessaging, Error, TEXT("StartTransport failed to create unicast socket on %s"), *UnicastEndpoint.ToString());
-
 		return false;
 	}
+	else
+	{
+		int32 PortNo = UnicastSocket->GetPortNo();
+		UE_LOG(LogUdpMessaging, Display, TEXT("Unicast socket bound to '%s:%d'."), *UnicastEndpoint.Address.ToString(), PortNo);
+	}
 #endif
 
 	// create & initialize multicast socket (optional)
