Add cvar to clear lost packets from the inflight list when in an overflow situation.

diff --git a/Engine/Plugins/Messaging/UdpMessaging/Source/UdpMessaging/Private/Transport/UdpMessageProcessor.cpp b/Engine/Plugins/Messaging/UdpMessaging/Source/UdpMessaging/Private/Transport/UdpMessageProcessor.cpp
--- a/Engine/Plugins/Messaging/UdpMessaging/Source/UdpMessaging/Private/Transport/UdpMessageProcessor.cpp
+++ b/Engine/Plugins/Messaging/UdpMessaging/Source/UdpMessaging/Private/Transport/UdpMessageProcessor.cpp
@@ -41,6 +41,13 @@ TAutoConsoleVariable<FString> CVarConnectionsToError(
 	TEXT("This can be a comma separated list in the form IPAddr2:port,IPAddr3:port"),
 	ECVF_Default);
 
+TAutoConsoleVariable<int32> CVarCheckForExpiredWithFullQueue(
+	TEXT("MessageBus.UDP.CheckForExpiredWithFullQueue"),
+	0,
+	TEXT("Attempts to release pressure on the work queue by checking if inflight segments have expired with no acknowledgement.\n"),
+	ECVF_Default);
+
+
 namespace UE::Private::MessageProcessor
 {
 
@@ -1251,7 +1258,6 @@ int32 FUdpMessageProcessor::UpdateSegmenters(FNodeInfo& NodeInfo)
 			return -1;
 		}
 	}
-	NodeInfo.OverflowForPendingAck.Reset();
 
 	// Process messages in the work queue.
 	while (NodeInfo.CanSendSegments() && NodeInfo.WorkQueue.Dequeue(MessageId))
@@ -1278,6 +1284,14 @@ int32 FUdpMessageProcessor::UpdateSegmenters(FNodeInfo& NodeInfo)
 		BytesSent += Info.BytesSent;
 	}
 
+	if (CVarCheckForExpiredWithFullQueue.GetValueOnAnyThread() > 0 && !NodeInfo.CanSendSegments() && NodeInfo.WorkQueue.IsFull())
+	{
+		UE_LOG(LogUdpMessaging, Warning, TEXT("Work queue is full sending to node at address %s. We cannot send new data. Attempting to expire old inflight segments."), *NodeInfo.Endpoint.ToString());
+		// RemoveLostSegments will clean-up the tracking map that we use to determine what is "inflight". Old data will be expired and scheduled
+		// for resend.
+		NodeInfo.RemoveLostSegments(CurrentTime);
+	}
+
 	NodeInfo.Statistics.PacketsInFlight = NodeInfo.InflightSegments.Num();
 	NodeInfo.Statistics.BytesInflight = NodeInfo.InflightSegments.Num() * UDP_MESSAGING_SEGMENT_SIZE;
 
