Cache the message type so that it is ready on deserialization. This allows us to deliver a message during package saves / GC.

diff --git a/Engine/Plugins/Messaging/UdpMessaging/Source/UdpMessaging/Private/Transport/UdpDeserializedMessage.cpp b/Engine/Plugins/Messaging/UdpMessaging/Source/UdpMessaging/Private/Transport/UdpDeserializedMessage.cpp
--- a/Engine/Plugins/Messaging/UdpMessaging/Source/UdpMessaging/Private/Transport/UdpDeserializedMessage.cpp
+++ b/Engine/Plugins/Messaging/UdpMessaging/Source/UdpMessaging/Private/Transport/UdpDeserializedMessage.cpp
@@ -5,10 +5,13 @@
 #include "Async/TaskGraphInterfaces.h"
 #include "Backends/JsonStructDeserializerBackend.h"
 #include "Backends/CborStructDeserializerBackend.h"
+#include "Misc/AssertionMacros.h"
 #include "StructDeserializer.h"
 #include "Serialization/MemoryReader.h"
 #include "Transport/UdpReassembledMessage.h"
 #include "Shared/UdpMessagingSettings.h"
+#include "UObject/TopLevelAssetPath.h"
+#include "UObject/WeakObjectPtrTemplates.h"
 #include "UdpMessagingTracing.h"
 
 /* FUdpDeserializedMessage structors
@@ -43,6 +46,9 @@ public:
 	static bool DeserializeV11_15(FUdpDeserializedMessage& DeserializedMessage, FMemoryReader& MessageReader, bool bIsLWCBackwardCompatibilityMode);
 	static bool DeserializeV16_17(FUdpDeserializedMessage& DeserializedMessage, FMemoryReader& MessageReader);
 	static bool Deserialize(FUdpDeserializedMessage& DeserializedMessage, const FUdpReassembledMessage& ReassembledMessage);
+
+	template <typename ReturnType>
+	static ReturnType GetMessageType(FMemoryReader& MessageReader);
 };
 
 /* FUdpDeserializedMessage interface
@@ -144,16 +150,28 @@ const FDateTime& FUdpDeserializedMessage::GetTimeSent() const
 /* FUdpDeserializedMessageDetails implementation
 *****************************************************************************/
 
+template <typename ReturnType>
+ReturnType FUdpDeserializedMessageDetails::GetMessageType(FMemoryReader& MessageReader)
+{
+	// This should only be used for TopLevelAssetPath or FName paths
+	static_assert(std::is_same<FName, ReturnType>::value
+				  || std::is_same<FTopLevelAssetPath, ReturnType>::value);
+
+	ReturnType MessageType;
+	MessageReader << MessageType;
+	return MessageType;
+}
 
 bool FUdpDeserializedMessageDetails::DeserializeV10(FUdpDeserializedMessage& DeserializedMessage, FMemoryReader& MessageReader)
 {
 	// message type info
 	{
-		FName MessageType;
-		MessageReader << MessageType;
+		FName MessageType = GetMessageType<FName>(MessageReader);
 
-		// @todo gmp: cache message types for faster lookup
-		DeserializedMessage.TypeInfo = FindFirstObjectSafe<UScriptStruct>(*MessageType.ToString(), EFindFirstObjectOptions::EnsureIfAmbiguous);
+		if (DeserializedMessage.TypeInfo == nullptr)
+		{
+			DeserializedMessage.TypeInfo = FindFirstObjectSafe<UScriptStruct>(*MessageType.ToString(), EFindFirstObjectOptions::EnsureIfAmbiguous);
+		}
 
 		if (!DeserializedMessage.TypeInfo.IsValid(false, true))
 		{
@@ -235,11 +253,12 @@ bool FUdpDeserializedMessageDetails::DeserializeV11_15(FUdpDeserializedMessage&
 {
 	// message type info
 	{
-		FName MessageType;
-		MessageReader << MessageType;
+		FName MessageType = GetMessageType<FName>(MessageReader);
 
-		// @todo gmp: cache message types for faster lookup
-		DeserializedMessage.TypeInfo = FindFirstObjectSafe<UScriptStruct>(*MessageType.ToString(), EFindFirstObjectOptions::EnsureIfAmbiguous);
+		if (DeserializedMessage.TypeInfo == nullptr)
+		{
+			DeserializedMessage.TypeInfo = FindFirstObjectSafe<UScriptStruct>(*MessageType.ToString(), EFindFirstObjectOptions::EnsureIfAmbiguous);
+		}
 
 		if (!DeserializedMessage.TypeInfo.IsValid(false, true))
 		{
@@ -364,11 +383,12 @@ bool FUdpDeserializedMessageDetails::DeserializeV16_17(FUdpDeserializedMessage&
 {
 	// message type info
 	{
-		FTopLevelAssetPath MessageType;
-		MessageReader << MessageType;
+		FTopLevelAssetPath MessageType = GetMessageType<FTopLevelAssetPath>(MessageReader);
 
-		// @todo gmp: cache message types for faster lookup
-		DeserializedMessage.TypeInfo = FindObjectSafe<UScriptStruct>(MessageType);
+		if (DeserializedMessage.TypeInfo == nullptr)
+		{
+			DeserializedMessage.TypeInfo = FindObjectSafe<UScriptStruct>(MessageType);
+		}
 
 		if (!DeserializedMessage.TypeInfo.IsValid(false, true))
 		{
@@ -490,6 +510,58 @@ bool FUdpDeserializedMessageDetails::DeserializeV16_17(FUdpDeserializedMessage&
 	}
 }
 
+FNameOrFTopLevel FUdpDeserializedMessage::PeekMessageTypeInfoName(const FUdpReassembledMessage& PartialMessage)
+{
+	check(PartialMessage.HasFirstSegment());
+	FNameOrFTopLevel Path;
+
+	FMemoryReader MessageReader(PartialMessage.GetData());
+	switch (PartialMessage.GetProtocolVersion())
+	{
+	case 10:
+	case 11:
+	case 12:
+	case 13:
+	case 14:
+	case 15:
+	{
+		FName MessageName = FUdpDeserializedMessageDetails::GetMessageType<FName>(MessageReader);
+		Path.Set<FName>(MessageName);
+		break;
+	}
+	case 16:
+		// fallthrough
+	case 17:
+	{
+		FTopLevelAssetPath AssetPath = FUdpDeserializedMessageDetails::GetMessageType<FTopLevelAssetPath>(MessageReader);
+		Path.Set<FTopLevelAssetPath>(AssetPath);
+		break;
+	}
+	default:
+		UE_LOG(LogUdpMessaging, Error, TEXT("Unsupported Protocol Version message tasked for MessageType query."));
+		break;
+	};
+	return Path;
+}
+
+TWeakObjectPtr<UScriptStruct> FUdpDeserializedMessage::ResolvePath(const FNameOrFTopLevel& Path)
+{
+	TWeakObjectPtr<UScriptStruct> TypeInfo;
+	if (Path.IsType<FName>())
+	{
+		TypeInfo = FindFirstObjectSafe<UScriptStruct>(*Path.Get<FName>().ToString(), EFindFirstObjectOptions::EnsureIfAmbiguous);
+	}
+	else if (Path.IsType<FTopLevelAssetPath>())
+	{
+		TypeInfo = FindObjectSafe<UScriptStruct>(Path.Get<FTopLevelAssetPath>());
+	}
+	else
+	{
+		checkNoEntry();
+	}
+	return TypeInfo;
+}
+
 bool FUdpDeserializedMessageDetails::Deserialize(FUdpDeserializedMessage& DeserializedMessage, const FUdpReassembledMessage& ReassembledMessage)
 {
 	// Note that some complex values are deserialized manually here, so that we
@@ -498,6 +570,10 @@ bool FUdpDeserializedMessageDetails::Deserialize(FUdpDeserializedMessage& Deseri
 	FMemoryReader MessageReader(ReassembledMessage.GetData());
 	MessageReader.ArMaxSerializeSize = NAME_SIZE;
 
+	// We may be able to skip lookup of the type name based on precached data provided by
+	// the reassembler.
+	//
+	DeserializedMessage.TypeInfo = ReassembledMessage.GetMessageTypeInfo();
 	switch (ReassembledMessage.GetProtocolVersion())
 	{
 	case 10:
diff --git a/Engine/Plugins/Messaging/UdpMessaging/Source/UdpMessaging/Private/Transport/UdpDeserializedMessage.h b/Engine/Plugins/Messaging/UdpMessaging/Source/UdpMessaging/Private/Transport/UdpDeserializedMessage.h
--- a/Engine/Plugins/Messaging/UdpMessaging/Source/UdpMessaging/Private/Transport/UdpDeserializedMessage.h
+++ b/Engine/Plugins/Messaging/UdpMessaging/Source/UdpMessaging/Private/Transport/UdpDeserializedMessage.h
@@ -4,12 +4,15 @@
 
 #include "IMessageContext.h"
 #include "Misc/DateTime.h"
+#include "UObject/TopLevelAssetPath.h"
+#include "Misc/TVariant.h"
 
 class FUdpReassembledMessage;
 class IMessageAttachment;
 class UScriptStruct;
 class FUdpDeserializedMessageDetails;
 
+using FNameOrFTopLevel = TVariant<FName, FTopLevelAssetPath>;
 /**
  * Holds a deserialized message.
  */
@@ -38,6 +41,12 @@ public:
 	 */
 	bool Deserialize(const FUdpReassembledMessage& ReassembledMessage);
 
+	/**
+	 * Called early by the message processor to peek at the typeinfo name.
+	 */
+	static FNameOrFTopLevel PeekMessageTypeInfoName(const FUdpReassembledMessage& ReassembledMessage);
+	static TWeakObjectPtr<UScriptStruct> ResolvePath(const FNameOrFTopLevel& Path);
+
 public:
 
 	//~ IMessageContext interface
diff --git a/Engine/Plugins/Messaging/UdpMessaging/Source/UdpMessaging/Private/Transport/UdpMessageProcessor.cpp b/Engine/Plugins/Messaging/UdpMessaging/Source/UdpMessaging/Private/Transport/UdpMessageProcessor.cpp
--- a/Engine/Plugins/Messaging/UdpMessaging/Source/UdpMessaging/Private/Transport/UdpMessageProcessor.cpp
+++ b/Engine/Plugins/Messaging/UdpMessaging/Source/UdpMessaging/Private/Transport/UdpMessageProcessor.cpp
@@ -7,6 +7,7 @@
 #include "Serialization/ArrayReader.h"
 #include "Serialization/ArrayWriter.h"
 
+#include "Transport/UdpDeserializedMessage.h"
 #include "UdpCircularQueue.h"
 #include "UdpMessagingTracing.h"
 #include "Shared/UdpMessagingSettings.h"
@@ -908,10 +909,36 @@ void FUdpMessageProcessor::ProcessUnknownSegment(FInboundSegment& Segment, FNode
 	UE_LOG(LogUdpMessaging, Verbose, TEXT("Received unknown segment type '%i' from %s"), SegmentType, *Segment.Sender.ToText().ToString());
 }
 
+void FUdpMessageProcessor::LookupAndCacheMessageType(TSharedPtr<FUdpReassembledMessage, ESPMode::ThreadSafe>& ReassembledMessage)
+{
+	if (!ReassembledMessage->HasFirstSegment())
+	{
+		return;
+	}
+	if (ReassembledMessage->GetMessageTypeInfo() == nullptr)
+	{
+		FNameOrFTopLevel AssetPath = FUdpDeserializedMessage::PeekMessageTypeInfoName(*ReassembledMessage);
+		FString PathAsString = Visit([](auto&& Path) -> FString { return Path.ToString(); }, AssetPath);
+
+		if (TWeakObjectPtr<UScriptStruct>* TypeInfo = CachedTypeInfoMap.Find(PathAsString))
+		{
+			ReassembledMessage->SetMessageTypeInfo(*TypeInfo);
+		}
+		else if (!GIsSavingPackage && !IsGarbageCollecting())
+		{
+			// Otherwise we have to look up the object by calling FindObjectSafe.  This can fail in GC and package save.
+			// Thus we only do this in not saving / GC cases.
+			TWeakObjectPtr<UScriptStruct> Obj = FUdpDeserializedMessage::ResolvePath(AssetPath);
+			CachedTypeInfoMap.Add(PathAsString, Obj);
+			ReassembledMessage->SetMessageTypeInfo(MoveTemp(Obj));
+		}
+	}
+}
+
 void FUdpMessageProcessor::DeliverMessage(const TSharedPtr<FUdpReassembledMessage, ESPMode::ThreadSafe>& ReassembledMessage, FNodeInfo& NodeInfo)
 {
 	// Do not deliver message while saving or garbage collecting since those deliveries will fail anyway...
-	if (GIsSavingPackage || IsGarbageCollecting())
+	if (ReassembledMessage->GetMessageTypeInfo() == nullptr && (GIsSavingPackage || IsGarbageCollecting()))
 	{
 		UE_LOG(LogUdpMessaging, Verbose, TEXT("Skipping delivery of %s"), *ReassembledMessage->Describe());
 		return;
@@ -1350,6 +1377,8 @@ bool FUdpMessageProcessor::UpdateReassemblers(FNodeInfo& NodeInfo)
 				*NodeInfo.NodeId.ToString());
 		}
 
+		LookupAndCacheMessageType(ReassembledMessage);
+
 		// Try to deliver completed message that couldn't be delivered the first time around
 		if (ReassembledMessage->IsComplete() && !ReassembledMessage->IsDelivered())
 		{
diff --git a/Engine/Plugins/Messaging/UdpMessaging/Source/UdpMessaging/Private/Transport/UdpMessageProcessor.h b/Engine/Plugins/Messaging/UdpMessaging/Source/UdpMessaging/Private/Transport/UdpMessageProcessor.h
--- a/Engine/Plugins/Messaging/UdpMessaging/Source/UdpMessaging/Private/Transport/UdpMessageProcessor.h
+++ b/Engine/Plugins/Messaging/UdpMessaging/Source/UdpMessaging/Private/Transport/UdpMessageProcessor.h
@@ -12,6 +12,10 @@
 #include "Interfaces/IPv4/IPv4Endpoint.h"
 #include "Misc/SingleThreadRunnable.h"
 
+#include "UObject/TopLevelAssetPath.h"
+#include "UObject/WeakObjectPtr.h"
+#include "UObject/WeakObjectPtrTemplates.h"
+
 #include "UdpMessageSegmenter.h"
 #include "UdpMessagingPrivate.h"
 #include "Transport/UdpCircularQueue.h"
@@ -801,6 +805,11 @@ protected:
 	 */
 	void DeliverMessage(const TSharedPtr<FUdpReassembledMessage, ESPMode::ThreadSafe>& ReassembledMessage, FNodeInfo& NodeInfo);
 
+	/**
+	 * Lookup the reassembled message type and try to determine the source struct type.
+	 */
+	void LookupAndCacheMessageType(TSharedPtr<FUdpReassembledMessage, ESPMode::ThreadSafe>& ReassembledMessage);
+
 	/**
 	 * Removes the specified node from the list of known remote endpoints.
 	 *
@@ -950,6 +959,9 @@ private:
 	/** If our round-robin work queues couldn't accept the last outbound message then store a deferred message for next round. */
 	TOptional<FOutboundMessage> DeferredOutboundMessage;
 
+	/** Stores a cache of discovered UScriptStruct types found during message processing. */
+	TMap<FString, TWeakObjectPtr<UScriptStruct>> CachedTypeInfoMap;
+
 	/** Defines the maximum number of Hello segments that can be dropped before a remote endpoint is considered dead. */
 	static const int32 DeadHelloIntervals;
 
diff --git a/Engine/Plugins/Messaging/UdpMessaging/Source/UdpMessaging/Private/Transport/UdpReassembledMessage.h b/Engine/Plugins/Messaging/UdpMessaging/Source/UdpMessaging/Private/Transport/UdpReassembledMessage.h
--- a/Engine/Plugins/Messaging/UdpMessaging/Source/UdpMessaging/Private/Transport/UdpReassembledMessage.h
+++ b/Engine/Plugins/Messaging/UdpMessaging/Source/UdpMessaging/Private/Transport/UdpReassembledMessage.h
@@ -239,6 +239,18 @@ public:
 	}
 
 
+	/**
+	 * Indicates if the first segment of the data is ready.
+	 **/
+	bool HasFirstSegment() const
+	{
+		if (PendingSegments.Num() > 0)
+		{
+			return PendingSegments[0] == false;
+		}
+		return false;
+	}
+
 	/**
 	 * Reassembles a segment into the specified message.
 	 *
@@ -305,7 +317,18 @@ public:
 		return ReceivedBytes;
 	}
 
+	TWeakObjectPtr<UScriptStruct> GetMessageTypeInfo() const
+	{
+		return TypeInfo;
+	}
+
+	void SetMessageTypeInfo(TWeakObjectPtr<UScriptStruct> InTypeInfo)
+	{
+		TypeInfo = MoveTemp(InTypeInfo);
+	}
 private:
+	/** Holds the message's type information. */
+	TWeakObjectPtr<UScriptStruct> TypeInfo;
 
 	/** Holds the message protocol version. */
 	uint8 ProtocolVersion;
