Protect access to BeaconInterval across threads by adding an atomic.

diff --git a/Engine/Plugins/Messaging/UdpMessaging/Source/UdpMessaging/Private/Transport/UdpMessageBeacon.h b/Engine/Plugins/Messaging/UdpMessaging/Source/UdpMessaging/Private/Transport/UdpMessageBeacon.h
--- a/Engine/Plugins/Messaging/UdpMessaging/Source/UdpMessaging/Private/Transport/UdpMessageBeacon.h
+++ b/Engine/Plugins/Messaging/UdpMessaging/Source/UdpMessaging/Private/Transport/UdpMessageBeacon.h
@@ -124,7 +124,7 @@ private:
 	void AddLocalEndpoints();
 
 	/** Holds the calculated interval between Hello segments. */
-	FTimespan BeaconInterval;
+	std::atomic<FTimespan> BeaconInterval;
 
 	/** Holds an event signaling that an endpoint left. */
 	FEvent* EndpointLeftEvent;
