Fix TimeSynchronizationWriter and LiveLinkProviderSample programs

diff --git a/Engine/Plugins/Messaging/UdpMessaging/UdpMessaging.uplugin b/Engine/Plugins/Messaging/UdpMessaging/UdpMessaging.uplugin
--- a/Engine/Plugins/Messaging/UdpMessaging/UdpMessaging.uplugin
+++ b/Engine/Plugins/Messaging/UdpMessaging/UdpMessaging.uplugin
@@ -23,7 +23,8 @@
 		"UnrealLightmass",
 		"CrashReportClientEditor",
 		"LiveLinkHub",
-		"TimeSynchronizationWriter"
+		"TimeSynchronizationWriter",
+        "LiveLinkProviderSample"
 	],
 	"Modules" :
 	[
@@ -43,7 +44,8 @@
 				"UnrealLightmass",
 				"CrashReportClientEditor",
 				"LiveLinkHub",
-				"TimeSynchronizationWriter"
+				"TimeSynchronizationWriter",
+                "LiveLinkProviderSample"
 			]
 		}
 	]
