Remove uses of TQueue and replace with TMpscQueue or TSpscQueue where necessary.

#rnx

diff --git a/Engine/Plugins/Messaging/UdpMessaging/Source/UdpMessaging/Private/Transport/UdpMessageBeacon.cpp b/Engine/Plugins/Messaging/UdpMessaging/Source/UdpMessaging/Private/Transport/UdpMessageBeacon.cpp
--- a/Engine/Plugins/Messaging/UdpMessaging/Source/UdpMessaging/Private/Transport/UdpMessageBeacon.cpp
+++ b/Engine/Plugins/Messaging/UdpMessaging/Source/UdpMessaging/Private/Transport/UdpMessageBeacon.cpp
@@ -82,12 +82,12 @@ bool FUdpMessageBeacon::HasSocketError() const
 
 void FUdpMessageBeacon::AddStaticEndpoint(const FIPv4Endpoint& InEndpoint)
 {
-	StaticEndpointQueue.Enqueue({ InEndpoint, true });
+	StaticEndpointQueue.Enqueue(FPendingEndpoint( InEndpoint, true));
 }
 
 void FUdpMessageBeacon::RemoveStaticEndpoint(const FIPv4Endpoint& InEndpoint)
 {
-	StaticEndpointQueue.Enqueue({ InEndpoint, false });
+	StaticEndpointQueue.Enqueue(FPendingEndpoint(InEndpoint, false));
 }
 
 /* FRunnable interface
diff --git a/Engine/Plugins/Messaging/UdpMessaging/Source/UdpMessaging/Private/Transport/UdpMessageBeacon.h b/Engine/Plugins/Messaging/UdpMessaging/Source/UdpMessaging/Private/Transport/UdpMessageBeacon.h
--- a/Engine/Plugins/Messaging/UdpMessaging/Source/UdpMessaging/Private/Transport/UdpMessageBeacon.h
+++ b/Engine/Plugins/Messaging/UdpMessaging/Source/UdpMessaging/Private/Transport/UdpMessageBeacon.h
@@ -7,7 +7,7 @@
 #include "Misc/SingleThreadRunnable.h"
 #include "HAL/Runnable.h"
 #include "Interfaces/IPv4/IPv4Endpoint.h"
-#include "Containers/Queue.h"
+#include "Containers/SpscQueue.h"
 
 enum class EUdpMessageSegments : uint8;
 
@@ -146,7 +146,7 @@ private:
 	};
 
 	/** Static endpoint queue to add/remove static endpoints. */
-	TQueue<FPendingEndpoint> StaticEndpointQueue;
+	TSpscQueue<FPendingEndpoint> StaticEndpointQueue;
 
 	/** Holds the static addresses to broadcast ping to. */
 	TArray<TSharedPtr<FInternetAddr>> StaticAddresses;
diff --git a/Engine/Plugins/Messaging/UdpMessaging/Source/UdpMessaging/Private/Transport/UdpMessageProcessor.cpp b/Engine/Plugins/Messaging/UdpMessaging/Source/UdpMessaging/Private/Transport/UdpMessageProcessor.cpp
--- a/Engine/Plugins/Messaging/UdpMessaging/Source/UdpMessaging/Private/Transport/UdpMessageProcessor.cpp
+++ b/Engine/Plugins/Messaging/UdpMessaging/Source/UdpMessaging/Private/Transport/UdpMessageProcessor.cpp
@@ -245,10 +245,7 @@ bool FUdpMessageProcessor::EnqueueInboundSegment(const TSharedPtr<FArrayReader,
 		return false;
 	}
 
-	if (!InboundSegments.Enqueue(FInboundSegment(Data, InSender)))
-	{
-		return false;
-	}
+	InboundSegments.Enqueue(FInboundSegment(Data, InSender));
 
 	WorkEvent->Trigger();
 
@@ -272,10 +269,7 @@ bool FUdpMessageProcessor::EnqueueOutboundMessage(const TSharedRef<IMessageConte
 		TGraphTask<FUdpSerializeMessageTask>::CreateTask().ConstructAndDispatchWhenReady(MessageContext, SerializedMessage, WorkEvent);
 
 		// Enqueue the message
-		if (!OutboundMessages.Enqueue(FOutboundMessage(SerializedMessage, RecipientVersion.Value, MessageContext->GetFlags())))
-		{
-			return false;
-		}
+		OutboundMessages.Enqueue(FOutboundMessage(SerializedMessage, RecipientVersion.Value, MessageContext->GetFlags()));
 	}
 
 	return true;
diff --git a/Engine/Plugins/Messaging/UdpMessaging/Source/UdpMessaging/Private/Transport/UdpMessageProcessor.h b/Engine/Plugins/Messaging/UdpMessaging/Source/UdpMessaging/Private/Transport/UdpMessageProcessor.h
--- a/Engine/Plugins/Messaging/UdpMessaging/Source/UdpMessaging/Private/Transport/UdpMessageProcessor.h
+++ b/Engine/Plugins/Messaging/UdpMessaging/Source/UdpMessaging/Private/Transport/UdpMessageProcessor.h
@@ -2,7 +2,7 @@
 
 #pragma once
 
-#include "Containers/Queue.h"
+#include "Containers/MpscQueue.h"
 
 #include "ModuleDescriptor.h"
 
@@ -870,10 +870,10 @@ private:
 	bool ConsumeOneOutboundMessage(const FOutboundMessage& OutboundMessage);
 
 	/** Holds the queue of outbound messages. */
-	TQueue<FInboundSegment, EQueueMode::Mpsc> InboundSegments;
+	TMpscQueue<FInboundSegment> InboundSegments;
 
 	/** Holds the queue of outbound messages. */
-	TQueue<FOutboundMessage, EQueueMode::Mpsc> OutboundMessages;
+	TMpscQueue<FOutboundMessage> OutboundMessages;
 
 	/** Holds the hello sender. */
 	TUniquePtr<FUdpMessageBeacon> Beacon;
diff --git a/Engine/Plugins/Messaging/UdpMessaging/Source/UdpMessaging/Private/Tunnel/UdpMessageTunnel.h b/Engine/Plugins/Messaging/UdpMessaging/Source/UdpMessaging/Private/Tunnel/UdpMessageTunnel.h
--- a/Engine/Plugins/Messaging/UdpMessaging/Source/UdpMessaging/Private/Tunnel/UdpMessageTunnel.h
+++ b/Engine/Plugins/Messaging/UdpMessaging/Source/UdpMessaging/Private/Tunnel/UdpMessageTunnel.h
@@ -5,7 +5,7 @@
 #include "Misc/DateTime.h"
 #include "Misc/Guid.h"
 #include "HAL/Runnable.h"
-#include "Containers/Queue.h"
+#include "Containers/MpscQueue.h"
 #include "Interfaces/IPv4/IPv4Endpoint.h"
 #include "Misc/SingleThreadRunnable.h"
 
@@ -138,7 +138,7 @@ private:
 	FSocket* MulticastSocket;
 
 	/** Holds a queue of pending connections. */
-	TQueue<TSharedPtr<FUdpMessageTunnelConnection>, EQueueMode::Mpsc> PendingConnections;
+	TMpscQueue<TSharedPtr<FUdpMessageTunnelConnection>> PendingConnections;
 
 	/** Holds a collection of information for remote transport nodes. */
 	TMap<FGuid, FNodeInfo> RemoteNodes;
diff --git a/Engine/Plugins/Messaging/UdpMessaging/Source/UdpMessaging/Private/Tunnel/UdpMessageTunnelConnection.cpp b/Engine/Plugins/Messaging/UdpMessaging/Source/UdpMessaging/Private/Tunnel/UdpMessageTunnelConnection.cpp
--- a/Engine/Plugins/Messaging/UdpMessaging/Source/UdpMessaging/Private/Tunnel/UdpMessageTunnelConnection.cpp
+++ b/Engine/Plugins/Messaging/UdpMessaging/Source/UdpMessaging/Private/Tunnel/UdpMessageTunnelConnection.cpp
@@ -55,7 +55,8 @@ bool FUdpMessageTunnelConnection::Send(const TSharedPtr<FArrayReader, ESPMode::T
 {
 	if (IsOpen())
 	{
-		return Outbox.Enqueue(Payload);
+		Outbox.Enqueue(Payload);
+		return true;
 	}
 
 	return false;
diff --git a/Engine/Plugins/Messaging/UdpMessaging/Source/UdpMessaging/Private/Tunnel/UdpMessageTunnelConnection.h b/Engine/Plugins/Messaging/UdpMessaging/Source/UdpMessaging/Private/Tunnel/UdpMessageTunnelConnection.h
--- a/Engine/Plugins/Messaging/UdpMessaging/Source/UdpMessaging/Private/Tunnel/UdpMessageTunnelConnection.h
+++ b/Engine/Plugins/Messaging/UdpMessaging/Source/UdpMessaging/Private/Tunnel/UdpMessageTunnelConnection.h
@@ -2,7 +2,7 @@
 
 #pragma once
 
-#include "Containers/Queue.h"
+#include "Containers/SpscQueue.h"
 #include "HAL/Runnable.h"
 #include "Interfaces/IPv4/IPv4Endpoint.h"
 #include "Misc/DateTime.h"
@@ -98,13 +98,13 @@ private:
 	FDateTime ClosedTime;
 
 	/** Holds the collection of received payloads. */
-	TQueue<TSharedPtr<FArrayReader, ESPMode::ThreadSafe>> Inbox;
+	TSpscQueue<TSharedPtr<FArrayReader, ESPMode::ThreadSafe>> Inbox;
 
 	/** Holds the time at which the connection was opened. */
 	FDateTime OpenedTime;
 
 	/** Holds the collection of outbound payloads. */
-	TQueue<TSharedPtr<FArrayReader, ESPMode::ThreadSafe>> Outbox;
+	TSpscQueue<TSharedPtr<FArrayReader, ESPMode::ThreadSafe>> Outbox;
 
 	/** Holds the IP endpoint of the remote client. */
 	FIPv4Endpoint RemoteEndpoint;
