Fix horde issue with MacArm64

#rnx

diff --git a/Engine/Plugins/Messaging/UdpMessaging/Source/UdpMessaging/Private/Transport/UdpMessageBeacon.cpp b/Engine/Plugins/Messaging/UdpMessaging/Source/UdpMessaging/Private/Transport/UdpMessageBeacon.cpp
--- a/Engine/Plugins/Messaging/UdpMessaging/Source/UdpMessaging/Private/Transport/UdpMessageBeacon.cpp
+++ b/Engine/Plugins/Messaging/UdpMessaging/Source/UdpMessaging/Private/Transport/UdpMessageBeacon.cpp
@@ -82,12 +82,14 @@ bool FUdpMessageBeacon::HasSocketError() const
 
 void FUdpMessageBeacon::AddStaticEndpoint(const FIPv4Endpoint& InEndpoint)
 {
-	StaticEndpointQueue.Enqueue(FPendingEndpoint( InEndpoint, true));
+	FPendingEndpoint EndPt{InEndpoint,true};
+	StaticEndpointQueue.Enqueue(MoveTemp(EndPt));
 }
 
 void FUdpMessageBeacon::RemoveStaticEndpoint(const FIPv4Endpoint& InEndpoint)
 {
-	StaticEndpointQueue.Enqueue(FPendingEndpoint(InEndpoint, false));
+	FPendingEndpoint EndPt{InEndpoint, false};
+	StaticEndpointQueue.Enqueue(MoveTemp(EndPt));
 }
 
 /* FRunnable interface
