[UdpSerializedMessage]
 - Fix race by making sure we release/acquire when changing/reading the state of the message

diff --git a/Engine/Plugins/Messaging/UdpMessaging/Source/UdpMessaging/Private/Transport/UdpSerializedMessage.h b/Engine/Plugins/Messaging/UdpMessaging/Source/UdpMessaging/Private/Transport/UdpSerializedMessage.h
--- a/Engine/Plugins/Messaging/UdpMessaging/Source/UdpMessaging/Private/Transport/UdpSerializedMessage.h
+++ b/Engine/Plugins/Messaging/UdpMessaging/Source/UdpMessaging/Private/Transport/UdpSerializedMessage.h
@@ -6,6 +6,8 @@
 #include "Serialization/MemoryReader.h"
 #include "IMessageContext.h"
 
+#include <atomic>
+
 enum class EUdpMessageFormat : uint8;
 
 /**
@@ -81,7 +83,7 @@ public:
 	 */
 	EUdpSerializedMessageState GetState() const
 	{
-		return State;
+		return State.load(std::memory_order_acquire);
 	}
 
 	/**
@@ -92,7 +94,7 @@ public:
 	 */
 	void UpdateState(EUdpSerializedMessageState InState)
 	{
-		State = InState;
+		State.store(InState, std::memory_order_release);
 	}
 
 	/** @return the message flags. */
@@ -118,7 +120,7 @@ private:
 	TArray<uint8> DataArray;
 
 	/** Holds the message data state. */
-	EUdpSerializedMessageState State;
+	std::atomic<EUdpSerializedMessageState> State;
 
 	/** Holds message flags, captured from context. */
 	EMessageFlags Flags;
