[ALT]
 - Fix some IsSavingPackage race by using the function instead of the global

#rnx

diff --git a/Engine/Plugins/Messaging/UdpMessaging/Source/UdpMessaging/Private/Transport/UdpMessageProcessor.cpp b/Engine/Plugins/Messaging/UdpMessaging/Source/UdpMessaging/Private/Transport/UdpMessageProcessor.cpp
--- a/Engine/Plugins/Messaging/UdpMessaging/Source/UdpMessaging/Private/Transport/UdpMessageProcessor.cpp
+++ b/Engine/Plugins/Messaging/UdpMessaging/Source/UdpMessaging/Private/Transport/UdpMessageProcessor.cpp
@@ -918,7 +918,7 @@ void FUdpMessageProcessor::LookupAndCacheMessageType(TSharedPtr<FUdpReassembledM
 		{
 			ReassembledMessage->SetMessageTypeInfo(*TypeInfo);
 		}
-		else if (!GIsSavingPackage && !IsGarbageCollecting())
+		else if (!UE::IsSavingPackage() && !IsGarbageCollecting())
 		{
 			// Otherwise we have to look up the object by calling FindObjectSafe.  This can fail in GC and package save.
 			// Thus we only do this in not saving / GC cases.
@@ -932,7 +932,7 @@ void FUdpMessageProcessor::LookupAndCacheMessageType(TSharedPtr<FUdpReassembledM
 void FUdpMessageProcessor::DeliverMessage(const TSharedPtr<FUdpReassembledMessage, ESPMode::ThreadSafe>& ReassembledMessage, FNodeInfo& NodeInfo)
 {
 	// Do not deliver message while saving or garbage collecting since those deliveries will fail anyway...
-	if (ReassembledMessage->GetMessageTypeInfo() == nullptr && (GIsSavingPackage || IsGarbageCollecting()))
+	if (ReassembledMessage->GetMessageTypeInfo() == nullptr && (UE::IsSavingPackage() || IsGarbageCollecting()))
 	{
 		UE_LOG(LogUdpMessaging, Verbose, TEXT("Skipping delivery of %s"), *ReassembledMessage->Describe());
 		return;
diff --git a/Engine/Source/Editor/UnrealEd/Private/EditorDomain/EditorDomainUtils.cpp b/Engine/Source/Editor/UnrealEd/Private/EditorDomain/EditorDomainUtils.cpp
--- a/Engine/Source/Editor/UnrealEd/Private/EditorDomain/EditorDomainUtils.cpp
+++ b/Engine/Source/Editor/UnrealEd/Private/EditorDomain/EditorDomainUtils.cpp
@@ -417,7 +417,7 @@ FPackageDigest CalculatePackageDigest(const FAssetPackageData& PackageData, FNam
 		if (NextClass < ImportedClassPaths.Num())
 		{
 			// EDITORDOMAIN_TODO: Remove the clauses !IsInGameThread || GIsSavingPackage once FindObject no longer asserts if GIsSavingPackage
-			if (bHasTriedPrecacheClassDigests || !IsInGameThread() || GIsSavingPackage)
+			if (bHasTriedPrecacheClassDigests || !IsInGameThread() || UE::IsSavingPackage())
 			{
 				return FPackageDigest(FPackageDigest::EStatus::MissingClass,
 					FName(WriteToString<256>(ImportedClassPaths[NextClass]).ToView()));
diff --git a/Engine/Source/Runtime/CoreUObject/Private/UObject/Obj.cpp b/Engine/Source/Runtime/CoreUObject/Private/UObject/Obj.cpp
--- a/Engine/Source/Runtime/CoreUObject/Private/UObject/Obj.cpp
+++ b/Engine/Source/Runtime/CoreUObject/Private/UObject/Obj.cpp
@@ -197,7 +197,7 @@ UObject* UObject::GetDefaultSubobjectByName(FName ToFind)
 {
 	UObject* Object = nullptr;
 	// If it is safe use the faster StaticFindObjectFast rather than searching all the subobjects
-	if (!GIsSavingPackage && !IsGarbageCollecting())
+	if (!UE::IsSavingPackage() && !IsGarbageCollecting())
 	{
 		Object = StaticFindObjectFast(UObject::StaticClass(), this, ToFind);
 		if (Object && !Object->IsDefaultSubobject())
diff --git a/Engine/Source/Runtime/CoreUObject/Private/UObject/SavePackage/SavePackageUtilities.cpp b/Engine/Source/Runtime/CoreUObject/Private/UObject/SavePackage/SavePackageUtilities.cpp
--- a/Engine/Source/Runtime/CoreUObject/Private/UObject/SavePackage/SavePackageUtilities.cpp
+++ b/Engine/Source/Runtime/CoreUObject/Private/UObject/SavePackage/SavePackageUtilities.cpp
@@ -183,7 +183,12 @@ namespace UE
 			// That global is only meant to be set and read from the game-thread...
 			// otherwise it could interfere with normal operations coming from
 			// other threads like async loading, etc...
-			return IsInGameThread() && GIsSavingPackage;
+			// We need to do a volatile read otherwise Clang/LLVM can decide to do
+			// some crazy optim where both GIsSavingPackage and IsInGameThread()
+			// are read into registry and evaluated together instead of going
+			// one after the other like we want to to avoid touching the GIsSavingPackage
+			// value from other threads and causing TSAN warnings.
+			return IsInGameThread() && *(volatile bool*)&GIsSavingPackage;
 		}
 		return InOuter->GetPackage()->HasAnyPackageFlags(PKG_IsSaving);
 	}
diff --git a/Engine/Source/Runtime/CoreUObject/Private/UObject/SoftObjectPath.cpp b/Engine/Source/Runtime/CoreUObject/Private/UObject/SoftObjectPath.cpp
--- a/Engine/Source/Runtime/CoreUObject/Private/UObject/SoftObjectPath.cpp
+++ b/Engine/Source/Runtime/CoreUObject/Private/UObject/SoftObjectPath.cpp
@@ -694,7 +694,7 @@ UObject* FSoftObjectPath::ResolveObject() const
 {
 	// Don't try to resolve if we're saving a package because StaticFindObject can't be used here
 	// and we usually don't want to force references to weak pointers while saving.
-	if (IsNull() || GIsSavingPackage)
+	if (IsNull() || UE::IsSavingPackage())
 	{
 		return nullptr;
 	}
diff --git a/Engine/Source/Runtime/Engine/Private/MemberReference.cpp b/Engine/Source/Runtime/Engine/Private/MemberReference.cpp
--- a/Engine/Source/Runtime/Engine/Private/MemberReference.cpp
+++ b/Engine/Source/Runtime/Engine/Private/MemberReference.cpp
@@ -362,7 +362,7 @@ TFieldType* FMemberReference::ResolveMemberImpl(UClass* SelfScope, TFieldTypeCla
 	else
 	{
 #if WITH_EDITOR
-		const bool bCanFollowRedirects = !GIsSavingPackage;
+		const bool bCanFollowRedirects = !UE::IsSavingPackage();
 #else
 		const bool bCanFollowRedirects = bAlwaysFollowRedirects;
 #endif
