Add debug log to udp messaging to help investigate a possible memory leak

#rnx

diff --git a/Engine/Plugins/Messaging/UdpMessaging/Source/UdpMessaging/Private/Transport/UdpMessageProcessor.cpp b/Engine/Plugins/Messaging/UdpMessaging/Source/UdpMessaging/Private/Transport/UdpMessageProcessor.cpp
--- a/Engine/Plugins/Messaging/UdpMessaging/Source/UdpMessaging/Private/Transport/UdpMessageProcessor.cpp
+++ b/Engine/Plugins/Messaging/UdpMessaging/Source/UdpMessaging/Private/Transport/UdpMessageProcessor.cpp
@@ -4,6 +4,8 @@
 #include "Algo/AllOf.h"
 
 #include "Common/UdpSocketSender.h"
+#include "CoreGlobals.h"
+#include "Misc/ConfigCacheIni.h"
 #include "Serialization/ArrayReader.h"
 #include "Serialization/ArrayWriter.h"
 
@@ -119,6 +121,8 @@ FUdpMessageProcessor::FUdpMessageProcessor(FSocket& InSocket, const FGuid& InNod
 {
 	Init();
 
+	MemoryLoggingThresholdMB = GConfig->GetFloatOrDefault(TEXT("/Script/UdpMessaging.UdpMessagingSettings"), TEXT("NodeMemoryLoggingThresholdMB"), -1, GEngineIni);
+
 	WorkEvent = MakeShareable(FPlatformProcess::GetSynchEventFromPool(), [](FEvent* EventToDelete)
 	{
 		FPlatformProcess::ReturnSynchEventToPool(EventToDelete);
@@ -1098,6 +1102,11 @@ void FUdpMessageProcessor::UpdateKnownNodes()
 		SendKnownNodesToKnownNodes();
 	}
 
+	float OutboundMessagesMemory = 0.0f;
+	float ReassemblersMemory = 0.0f;
+	constexpr uint32 LogFrequencySeconds = 5 * 60;
+	bool bCanLogDebugInfo = FPlatformTime::Seconds() > LastReassemblersLogTime + LogFrequencySeconds;
+
 	bool bSuccess = true;
 	for (auto& KnownNodePair : KnownNodes)
 	{
@@ -1120,6 +1129,28 @@ void FUdpMessageProcessor::UpdateKnownNodes()
 		}
 
 		KnownNodePair.Value.RemoveOldUnreliables(CurrentTime);
+
+		if (MemoryLoggingThresholdMB >= 0 && bCanLogDebugInfo)
+		{
+			uint32 ReassembledTotalBytes = 0;
+			for (const TPair<int32, TSharedPtr<FUdpReassembledMessage>>& Pair : KnownNodePair.Value.ReassembledMessages)
+			{
+				if (Pair.Value)
+				{
+					ReassembledTotalBytes += sizeof(FUdpReassembledMessage) + Pair.Value->GetReceivedBytes();
+				}
+			}
+
+			OutboundMessagesMemory += ReassembledTotalBytes / 1'000'000.f;
+			ReassemblersMemory += KnownNodePair.Value.Statistics.BytesInflight / 1'000'000.f;
+		}
+	}
+
+	// Only start logging once we accumulate more memory than what was specified in the UDP Messaging config.
+	if (MemoryLoggingThresholdMB >= 0 && OutboundMessagesMemory + ReassemblersMemory > MemoryLoggingThresholdMB && bCanLogDebugInfo)
+	{
+		UE_LOG(LogUdpMessaging, Log, TEXT("(Debug) MessageBus Memory Usage: Reassemblers: %.2f MB, Outbound messages: %.2f MB"), ReassemblersMemory, OutboundMessagesMemory);
+		LastReassemblersLogTime = FPlatformTime::Seconds();
 	}
 
 	// if we had socket error, fire up the error delegate
diff --git a/Engine/Plugins/Messaging/UdpMessaging/Source/UdpMessaging/Private/Transport/UdpMessageProcessor.h b/Engine/Plugins/Messaging/UdpMessaging/Source/UdpMessaging/Private/Transport/UdpMessageProcessor.h
--- a/Engine/Plugins/Messaging/UdpMessaging/Source/UdpMessaging/Private/Transport/UdpMessageProcessor.h
+++ b/Engine/Plugins/Messaging/UdpMessaging/Source/UdpMessaging/Private/Transport/UdpMessageProcessor.h
@@ -962,6 +962,12 @@ private:
 	/** Stores a cache of discovered UScriptStruct types found during message processing. */
 	TMap<FString, TWeakObjectPtr<UScriptStruct>> CachedTypeInfoMap;
 
+	/** How much memory (MB) should be used before periodically logging information about nodes' memory usage. */
+	float MemoryLoggingThresholdMB = -1;
+
+	/** Timestamp used to keep track of the last time we logged the active reassemblers. (A negative value will disable logging) */
+	double LastReassemblersLogTime = 0.0;
+
 	/** Defines the maximum number of Hello segments that can be dropped before a remote endpoint is considered dead. */
 	static const int32 DeadHelloIntervals;
 
diff --git a/Engine/Source/Programs/LiveLinkHub/Config/DefaultEngine.ini b/Engine/Source/Programs/LiveLinkHub/Config/DefaultEngine.ini
--- a/Engine/Source/Programs/LiveLinkHub/Config/DefaultEngine.ini
+++ b/Engine/Source/Programs/LiveLinkHub/Config/DefaultEngine.ini
@@ -58,6 +58,9 @@ ConverterDefaultPipeline=None
 BoneCompressionSettings="/Engine/Animation/DefaultAnimBoneCompressionSettings"
 CurveCompressionSettings="/Engine/Animation/DefaultAnimCurveCompressionSettings"
 
+[/Script/UdpMessaging.UdpMessagingSettings]
+NodeMemoryLoggingThresholdMB=5
+
 [CoreRedirects]
 +PropertyRedirects=(OldName="LiveLinkHubSettings.bAutoAddDiscoveredClients",NewName="AutoConnectClients")
 
