The deserialization of UDP messages happened on the processor thread which can cause delays in processing UDP messages. There was
a long standing TODO to move this into an async task.  Moving this to an async task will allow the processor to continue processing inbound
and outbounds messages and let the serialization and handoff of the message to the transport handler occur asynchronously.

diff --git a/Engine/Plugins/Messaging/UdpMessaging/Source/UdpMessaging/Private/Transport/UdpMessageProcessor.cpp b/Engine/Plugins/Messaging/UdpMessaging/Source/UdpMessaging/Private/Transport/UdpMessageProcessor.cpp
--- a/Engine/Plugins/Messaging/UdpMessaging/Source/UdpMessaging/Private/Transport/UdpMessageProcessor.cpp
+++ b/Engine/Plugins/Messaging/UdpMessaging/Source/UdpMessaging/Private/Transport/UdpMessageProcessor.cpp
@@ -944,7 +944,7 @@ void FUdpMessageProcessor::DeliverMessage(const TSharedPtr<FUdpReassembledMessag
 
 	if (NodeInfo.NodeId.IsValid())
 	{
-		MessageReassembledDelegate.ExecuteIfBound(*ReassembledMessage, nullptr, NodeInfo.NodeId);
+		MessageReassembledDelegate.ExecuteIfBound(ReassembledMessage, nullptr, NodeInfo.NodeId);
 	}
 
 	// Mark the message delivered but do not remove it from the list yet, this is to prevent the double delivery of reliable message
diff --git a/Engine/Plugins/Messaging/UdpMessaging/Source/UdpMessaging/Private/Transport/UdpMessageProcessor.h b/Engine/Plugins/Messaging/UdpMessaging/Source/UdpMessaging/Private/Transport/UdpMessageProcessor.h
--- a/Engine/Plugins/Messaging/UdpMessaging/Source/UdpMessaging/Private/Transport/UdpMessageProcessor.h
+++ b/Engine/Plugins/Messaging/UdpMessaging/Source/UdpMessaging/Private/Transport/UdpMessageProcessor.h
@@ -35,6 +35,7 @@ class FUdpSerializedMessage;
 class FUdpSocketSender;
 class IMessageAttachment;
 enum class EUdpMessageFormat : uint8;
+using FUdpReassembledMessageSharedPtr = TSharedPtr<FUdpReassembledMessage, ESPMode::ThreadSafe>;
 
 /**
  * Running statistics as known by the UdpMessageProcessor. This will be per endpoint.
@@ -588,7 +589,7 @@ public:
 	 *
 	 * @return The delegate.
 	 */
-	DECLARE_DELEGATE_ThreeParams(FOnMessageReassembled, const FUdpReassembledMessage& /*ReassembledMessage*/, const IMessageAttachmentPtr& /*Attachment*/, const FGuid& /*NodeId*/)
+	DECLARE_DELEGATE_ThreeParams(FOnMessageReassembled, FUdpReassembledMessageSharedPtr /*ReassembledMessage*/, const IMessageAttachmentPtr& /*Attachment*/, const FGuid& /*NodeId*/)
 	FOnMessageReassembled& OnMessageReassembled()
 	{
 		return MessageReassembledDelegate;
diff --git a/Engine/Plugins/Messaging/UdpMessaging/Source/UdpMessaging/Private/Transport/UdpMessageTransport.cpp b/Engine/Plugins/Messaging/UdpMessaging/Source/UdpMessaging/Private/Transport/UdpMessageTransport.cpp
--- a/Engine/Plugins/Messaging/UdpMessaging/Source/UdpMessaging/Private/Transport/UdpMessageTransport.cpp
+++ b/Engine/Plugins/Messaging/UdpMessaging/Source/UdpMessaging/Private/Transport/UdpMessageTransport.cpp
@@ -2,15 +2,20 @@
 
 #include "Transport/UdpMessageTransport.h"
 
+#include "HAL/IConsoleManager.h"
+
 #include "Common/UdpSocketBuilder.h"
 #include "Common/UdpSocketReceiver.h"
 #include "IMessageTransportHandler.h"
 #include "Async/Async.h"
 
 #include "Shared/UdpMessagingSettings.h"
+#include "Tasks/Task.h"
+#include "Tasks/TaskConcurrencyLimiter.h"
 #include "Transport/UdpReassembledMessage.h"
 #include "Transport/UdpDeserializedMessage.h"
 #include "Transport/UdpMessageProcessor.h"
+#include "UdpMessageProcessor.h"
 
 
 /* FUdpMessageTransport structors
@@ -309,6 +314,7 @@ bool FUdpMessageTransport::StartTransport(IMessageTransportHandler& Handler)
 #endif
 	}
 
+	
 	// initialize threads
 	FTimespan ThreadWaitTime = FTimespan::FromMilliseconds(100);
 
@@ -318,6 +324,13 @@ bool FUdpMessageTransport::StartTransport(IMessageTransportHandler& Handler)
 	MessageProcessor = new FUdpMessageProcessor(*MulticastSocket, FGuid::NewGuid(), MulticastEndpoint);
 #endif
 	{
+		UUdpMessagingSettings const* CurrentSettings = GetDefault<UUdpMessagingSettings>();
+		if (CurrentSettings->MaxConcurrentDeserializationTasks > 1)
+		{
+			UE_LOG(LogUdpMessaging, Display, TEXT("Using asynchronous task graph for message deserialization."));
+			DeserializationTasks = MakeUnique<UE::Tasks::FTaskConcurrencyLimiter>(CurrentSettings->MaxConcurrentDeserializationTasks);
+		}
+		
 		// Add the static endpoints
 		for (const FIPv4Endpoint& Endpoint : StaticEndpoints)
 		{
@@ -330,7 +343,6 @@ bool FUdpMessageTransport::StartTransport(IMessageTransportHandler& Handler)
 		MessageProcessor->OnError().BindRaw(this, &FUdpMessageTransport::HandleProcessorError);
 		MessageProcessor->OnErrorSendingToEndpoint_UdpMessageProcessorThread().BindRaw(this, &FUdpMessageTransport::HandleEndpointCommunicationError);
 		MessageProcessor->OnCanAcceptEndpoint_UdpMessageProcessorThread().BindRaw(this, &FUdpMessageTransport::HandleProcessorEndpointCheck);
-		UUdpMessagingSettings const * CurrentSettings = GetDefault<UUdpMessagingSettings>();
 		MessageProcessor->SetShareKnownNodesState(CurrentSettings->bShareKnownNodesWithActiveConnections);
 	}
 
@@ -387,6 +399,15 @@ void FUdpMessageTransport::StopTransport()
 	}
 #endif
 
+	// Clear out any deserialize tasks before resetting the transport handler. They depend on the handler to dispatch
+	// deserialized messages.
+	//
+	if (DeserializationTasks)
+	{
+		DeserializationTasks->Wait();
+		DeserializationTasks.Reset();		
+	}
+	
 	TransportHandler = nullptr;
 	ErrorFuture.Reset();
 	UE_LOG(LogUdpMessaging, Verbose, TEXT("Stopped Transport"));
@@ -418,18 +439,33 @@ bool FUdpMessageTransport::TransportMessage(const TSharedRef<IMessageContext, ES
 /* FUdpMessageTransport event handlers
  *****************************************************************************/
 
-void FUdpMessageTransport::HandleProcessorMessageReassembled(const FUdpReassembledMessage& ReassembledMessage, const TSharedPtr<IMessageAttachment, ESPMode::ThreadSafe>& Attachment, const FGuid& NodeId)
+void FUdpMessageTransport::HandleProcessorMessageReassembled(FUdpReassembledMessageSharedPtr ReassembledMessagePtr, const TSharedPtr<IMessageAttachment, ESPMode::ThreadSafe>& Attachment, const FGuid& NodeId)
 {
-	// @todo gmp: move message deserialization into an async task
 	TSharedRef<FUdpDeserializedMessage, ESPMode::ThreadSafe> DeserializedMessage = MakeShared<FUdpDeserializedMessage, ESPMode::ThreadSafe>(Attachment);
 
-	if (DeserializedMessage->Deserialize(ReassembledMessage))
+	auto DeserializeReassembledMessage = [DeserializedMessage = MoveTemp(DeserializedMessage), TransportHandler = this->TransportHandler, ReassembledMessagePtr = MoveTemp(ReassembledMessagePtr), NodeId](uint32 Slot)
+	{
+		const FUdpReassembledMessage& ReassembledMessage = *ReassembledMessagePtr.Get();
+		if (DeserializedMessage->Deserialize(ReassembledMessage))
+		{
+			TransportHandler->ReceiveTransportMessage(DeserializedMessage, NodeId);
+		}
+		else
+		{
+			UE_LOG(LogUdpMessaging, Verbose, TEXT("Failed to deserialize message from %s"), *NodeId.ToString());
+		}
+	};
+
+	if (DeserializationTasks)
 	{
-		TransportHandler->ReceiveTransportMessage(DeserializedMessage, NodeId);
+		// The user has asked for async deserialization.  Dispatch the task.
+		// 
+		DeserializationTasks->Push(UE_SOURCE_LOCATION, MoveTemp(DeserializeReassembledMessage));
 	}
 	else
 	{
-		UE_LOG(LogUdpMessaging, Verbose, TEXT("Failed to deserialize message from %s"), *NodeId.ToString());
+		// Otherwise we should run on this thread. 
+		DeserializeReassembledMessage(0);
 	}
 }
 
diff --git a/Engine/Plugins/Messaging/UdpMessaging/Source/UdpMessaging/Private/Transport/UdpMessageTransport.h b/Engine/Plugins/Messaging/UdpMessaging/Source/UdpMessaging/Private/Transport/UdpMessageTransport.h
--- a/Engine/Plugins/Messaging/UdpMessaging/Source/UdpMessaging/Private/Transport/UdpMessageTransport.h
+++ b/Engine/Plugins/Messaging/UdpMessaging/Source/UdpMessaging/Private/Transport/UdpMessageTransport.h
@@ -4,10 +4,12 @@
 
 #include "IMessageTransport.h"
 #include "Interfaces/IPv4/IPv4Endpoint.h"
+
 #include "Async/Future.h"
 #include "Containers/Ticker.h"
-#include "HAL/IConsoleManager.h"
 #include "Misc/Guid.h"
+#include "Templates/UniquePtr.h"
+#include "Tasks/TaskConcurrencyLimiter.h"
 
 struct FMessageTransportStatistics;
 
@@ -24,6 +26,7 @@ class ISocketSubsystem;
 
 struct FGuid;
 
+using FUdpReassembledMessageSharedPtr = TSharedPtr<FUdpReassembledMessage, ESPMode::ThreadSafe>;
 
 /**
  * Implements a message transport technology using an UDP network connection.
@@ -103,7 +106,7 @@ public:
 private:
 
 	/** Handles received transport messages. */
-	void HandleProcessorMessageReassembled(const FUdpReassembledMessage& ReassembledMessage, const TSharedPtr<IMessageAttachment, ESPMode::ThreadSafe>& Attachment, const FGuid& NodeId);
+	void HandleProcessorMessageReassembled(FUdpReassembledMessageSharedPtr ReassembledMessagePtr, const TSharedPtr<IMessageAttachment, ESPMode::ThreadSafe>& Attachment, const FGuid& NodeId);
 
 	/** Handles discovered transport endpoints. */
 	void HandleProcessorNodeDiscovered(const FGuid& DiscoveredNodeId);
@@ -182,4 +185,7 @@ private:
 
 	/** Excluded list of IP addresses as string value (can contain wildcards) not allowed to talk to UDP processor */
 	TArray<FString> ExcludedEndpoints;
+
+	/** Async task queue for running deserialization of reassembled messages. */
+	TUniquePtr<UE::Tasks::FTaskConcurrencyLimiter> DeserializationTasks;
 };
diff --git a/Engine/Plugins/Messaging/UdpMessaging/Source/UdpMessaging/Public/Shared/UdpMessagingSettings.h b/Engine/Plugins/Messaging/UdpMessaging/Source/UdpMessaging/Public/Shared/UdpMessagingSettings.h
--- a/Engine/Plugins/Messaging/UdpMessaging/Source/UdpMessaging/Public/Shared/UdpMessagingSettings.h
+++ b/Engine/Plugins/Messaging/UdpMessaging/Source/UdpMessaging/Public/Shared/UdpMessagingSettings.h
@@ -151,6 +151,13 @@ public:
 	UPROPERTY(config, EditAnywhere, Category=Transport, AdvancedDisplay)
 	bool bShareKnownNodesWithActiveConnections = false;
 
+	/**
+	 * Maximum number of deserialization tasks that can be performed in parallel. A value of 0 or 1 means that concurrent
+	 * deserialization is disabled, which will make the message processor behave similar to releases pre-5.6.
+	 */
+	UPROPERTY(config, EditAnywhere, Category=Transport, meta = (ClampMax="32"), AdvancedDisplay)
+	uint32 MaxConcurrentDeserializationTasks = 8;
+
 public:
 
 	/** Whether the UDP tunnel is enabled. */
diff --git a/Engine/Plugins/Messaging/UdpMessaging/Source/UdpMessaging/UdpMessaging.Build.cs b/Engine/Plugins/Messaging/UdpMessaging/Source/UdpMessaging/UdpMessaging.Build.cs
--- a/Engine/Plugins/Messaging/UdpMessaging/Source/UdpMessaging/UdpMessaging.Build.cs
+++ b/Engine/Plugins/Messaging/UdpMessaging/Source/UdpMessaging/UdpMessaging.Build.cs
@@ -11,6 +11,7 @@ namespace UnrealBuildTool.Rules
 				{
 					"Core",
 					"Networking",
+					"AtomicQueue"
 				});
 
 			PrivateDependencyModuleNames.AddRange(
