Remove a check on the version and report an error instead.  If we encouter this situation, discard the actual segment.  I believe this issues was caused by updating the protocol version everytime we talk to the endpoint. However, there are situation where the endpoint could send us a different protocol version.

My best guess for this happening is the following

a) We are running at the latest protocol version (17).  The client has a max N which N < 17
2) On ping / pong the client negotiates N-1 version number.
c) We respond to client with N-1 style messages.
d) Later the client sends us a N versioned message.  This would change the protocol version and invalidate any segments.

I believe we should always use the original negotiated version and never update.  I believe that the versioning change could have been a bug with previous versions of the engine.  Looking at the existing code, I do not see this as a potential issue going forward.  I'm just making the code more robust to handle this case.

#rnx

diff --git a/Engine/Plugins/Messaging/UdpMessaging/Source/UdpMessaging/Private/Transport/UdpMessageProcessor.cpp b/Engine/Plugins/Messaging/UdpMessaging/Source/UdpMessaging/Private/Transport/UdpMessageProcessor.cpp
--- a/Engine/Plugins/Messaging/UdpMessaging/Source/UdpMessaging/Private/Transport/UdpMessageProcessor.cpp
+++ b/Engine/Plugins/Messaging/UdpMessaging/Source/UdpMessaging/Private/Transport/UdpMessageProcessor.cpp
@@ -504,7 +504,6 @@ void FUdpMessageProcessor::ConsumeInboundSegments()
 				bAddedNewKnownNodes = true;
 			}
 
-			NodeInfo.ProtocolVersion = Header.ProtocolVersion;
 			NodeInfo.Endpoint = Segment.Sender;
 			NodeInfo.LastSegmentReceivedTime = CurrentTime;
 
@@ -626,6 +625,7 @@ bool FUdpMessageProcessor::ConsumeOneOutboundMessage(const FOutboundMessage& Out
 void FUdpMessageProcessor::ConsumeOutboundMessages()
 {
 	SCOPED_MESSAGING_TRACE(FUdpMessageProcessor_ConsumeOutputMessages);
+	
 	FOutboundMessage OutboundMessage;
 
 	if (DeferredOutboundMessage)
@@ -1220,8 +1220,12 @@ FSentSegmentInfo FUdpMessageProcessor::SendNextSegmentForMessageId(FNodeInfo& No
 	{
 		FUdpMessageSegment::FDataChunk DataChunk = GetDataChunk(NodeInfo, Segmenter, MessageId);
 
-		// validate with are sending message in the proper protocol version
-		check(Header.ProtocolVersion == Segmenter->GetProtocolVersion());
+		if (Header.ProtocolVersion != Segmenter->GetProtocolVersion())
+		{
+			UE_LOG(LogUdpMessaging, Error, TEXT("FUdpMessageProcessor::UpdateSegmenters: Node %s initially asked for protocol version %d but changed to %d. Cannot deliver message %d"), *NodeInfo.Endpoint.ToString(), Segmenter->GetProtocolVersion(), Header.ProtocolVersion, MessageId);
+			NodeInfo.Segmenters.Remove(MessageId);
+			return SentInfo;
+		}
 
 		TSharedRef<FArrayWriter, ESPMode::ThreadSafe> Writer = MakeShared<FArrayWriter, ESPMode::ThreadSafe>();
 		{
