Adjust UDP Messaging sender to fully utilize the available capacity in the send pipe and prioritize reliable messages over unreliable ones.

We also add settings for the prioritized queue and the time out period for known nodes.  The timeout period was fixed at 5 seconds and this gives users more control over its value.

[RN] Better utilization and prioritization of data sent over UDP messaging when reliable and unreliable data is sharing the transmission window.

diff --git a/Engine/Plugins/Messaging/UdpMessaging/Source/UdpMessaging/Private/Transport/UdpMessageProcessor.cpp b/Engine/Plugins/Messaging/UdpMessaging/Source/UdpMessaging/Private/Transport/UdpMessageProcessor.cpp
--- a/Engine/Plugins/Messaging/UdpMessaging/Source/UdpMessaging/Private/Transport/UdpMessageProcessor.cpp
+++ b/Engine/Plugins/Messaging/UdpMessaging/Source/UdpMessaging/Private/Transport/UdpMessageProcessor.cpp
@@ -11,6 +11,7 @@
 
 #include "Transport/UdpDeserializedMessage.h"
 #include "UdpCircularQueue.h"
+#include "UdpMessageProcessor.h"
 #include "UdpMessagingTracing.h"
 #include "Shared/UdpMessagingSettings.h"
 #include "Transport/UdpMessageBeacon.h"
@@ -103,7 +104,8 @@ FUdpMessageProcessor::FNodeInfo::FNodeInfo()
 	: LastSegmentReceivedTime(FDateTime::MinValue())
 	, NodeId()
 	, ProtocolVersion(UDP_MESSAGING_TRANSPORT_PROTOCOL_VERSION)
-	, WorkQueue(UE::Private::MessageProcessor::GetMessageProcessorWorkQueueSize())
+	, ReliableWorkQueue(UE::Private::MessageProcessor::GetMessageProcessorWorkQueueSize())
+	, UnreliableWorkQueue(UE::Private::MessageProcessor::GetMessageProcessorWorkQueueSize())  
 {
 	ComputeWindowSize(0,0);
 }
@@ -613,7 +615,14 @@ bool FUdpMessageProcessor::ConsumeOneOutboundMessage(const FOutboundMessage& Out
 			RecipientNodeInfo->Segmenters.Add(
 				LastSentMessage, Segmenter);
 
-			RecipientNodeInfo->WorkQueue.Enqueue(LastSentMessage);
+			if (bIsReliable)
+			{
+				RecipientNodeInfo->ReliableWorkQueue.Enqueue(LastSentMessage);				
+			}
+			else
+			{
+				RecipientNodeInfo->UnreliableWorkQueue.Enqueue(LastSentMessage);				
+			}
 			SendSegmenterStatsToListeners(LastSentMessage, RecipientNodeInfo->NodeId, Segmenter);
 		}
 	}
@@ -966,7 +975,7 @@ int32 GetMaxSendRate()
 
 bool HasTimedOut(const FDateTime& ReceivedTime, const FDateTime& CurrentTime)
 {
-	const int32 TimeoutInSeconds = 5;
+	const float TimeoutInSeconds = GetDefault<UUdpMessagingSettings>()->ConnectionTimeoutPeriod;
 	FTimespan TimeoutInterval = TimeoutInSeconds * FTimespan::FromMilliseconds(1000);
 	return (ReceivedTime + TimeoutInterval) < CurrentTime;
 }
@@ -975,6 +984,12 @@ void FUdpMessageProcessor::RemoveDeadNodes()
 {
 	// Remove dead nodes
 	FTimespan DeadHelloTimespan = DeadHelloIntervals * Beacon->GetBeaconInterval();
+	if (DeadHelloTimespan.IsZero())
+	{
+		// We never expire connections.
+		return;
+	}
+	
 	for (auto It = KnownNodes.CreateIterator(); It; ++It)
 	{
 		FGuid& NodeId = It->Key;
@@ -982,7 +997,7 @@ void FUdpMessageProcessor::RemoveDeadNodes()
 
 		if ((NodeId.IsValid()) && HasTimedOut(NodeInfo.LastSegmentReceivedTime,CurrentTime))
 		{
-			UE_LOG(LogUdpMessaging, Verbose, TEXT("FUdpMessageProcessor::UpdateKnownNodes: Removing Node %s (%s)"), *NodeInfo.NodeId.ToString(), *NodeInfo.Endpoint.ToString());
+			UE_LOG(LogUdpMessaging, Display, TEXT("FUdpMessageProcessor::UpdateKnownNodes: Removing Node %s (%s) because it timed-out."), *NodeInfo.NodeId.ToString(), *NodeInfo.Endpoint.ToString());
 			NodeLostDelegate.ExecuteIfBound(NodeId);
 			It.RemoveCurrent();
 		}
@@ -1281,9 +1296,8 @@ FSentSegmentInfo FUdpMessageProcessor::SendNextSegmentForMessageId(FNodeInfo& No
 	return MoveTemp(SentInfo);
 }
 
-int32 FUdpMessageProcessor::UpdateSegmenters(FNodeInfo& NodeInfo)
+TOptional<FSentSegmentInfoArray> FUdpMessageProcessor::DrainQueue(FNodeInfo& NodeInfo, TUdpCircularQueue<int32>& WorkQueue, int64 DrainCount, const bool bDrainOverflow)
 {
-	SCOPED_MESSAGING_TRACE(FUdpMessageProcessor_UpdateSegmenters);
 	FUdpMessageSegment::FHeader Header
 	{
 		NodeInfo.ProtocolVersion,		// Header.ProtocolVersion - Send data segment using the node protocol version
@@ -1292,55 +1306,120 @@ int32 FUdpMessageProcessor::UpdateSegmenters(FNodeInfo& NodeInfo)
 		EUdpMessageSegments::Data		// Header.SegmentType
 	};
 
-	int32 MessageId;
 	int32 BytesSent = 0;
-	TArray<FSentSegmentInfo,TInlineAllocator<64>> ProcessedMessages;
-	auto SendOne = [this, &NodeInfo, &Header, &ProcessedMessages](int32 MessageId)
-	{
-		FSentSegmentInfo Info = SendNextSegmentForMessageId(NodeInfo, Header, MessageId);
-		ProcessedMessages.Emplace(MoveTemp(Info));
-	};
-	auto SocketError = [&ProcessedMessages]()
-	{
-		return ProcessedMessages.Top().bSendSocketError;
-	};
+	FSentSegmentInfoArray ProcessedMessages;
 
-	// Check for acknowledged segments first.
-	while (NodeInfo.CanSendSegments() && !NodeInfo.OverflowForPendingAck.IsEmpty())
+	if (bDrainOverflow)
 	{
-		SendOne(NodeInfo.OverflowForPendingAck.Pop());
-		if (SocketError())
+		// Check for acknowledged segments first.
+		while (NodeInfo.CanSendSegments() && !NodeInfo.OverflowForPendingAck.IsEmpty())
 		{
-			return -1;
+			FSentSegmentInfo Info = SendNextSegmentForMessageId(NodeInfo, Header, NodeInfo.OverflowForPendingAck.Pop());
+			if (Info.bSendSocketError)
+			{
+				return {};
+			}
+			ProcessedMessages.Emplace(MoveTemp(Info));
+			DrainCount--;
 		}
 	}
-
+	
+	int32 MessageId;
 	// Process messages in the work queue.
-	while (NodeInfo.CanSendSegments() && NodeInfo.WorkQueue.Dequeue(MessageId))
+	while (DrainCount > 0 && NodeInfo.CanSendSegments() && WorkQueue.Dequeue(MessageId))
 	{
-		SendOne(MessageId);
-		if (SocketError())
+		FSentSegmentInfo Info = SendNextSegmentForMessageId(NodeInfo, Header, MessageId);
+		if (Info.bSendSocketError)
 		{
-			return -1;
+			return {};
 		}
+		ProcessedMessages.Emplace(MoveTemp(Info));
+		DrainCount--;
 	}
+	return ProcessedMessages;
+}
 
+int32 FUdpMessageProcessor::RequeueWork(FNodeInfo& NodeInfo, TUdpCircularQueue<int32>& WorkQueue, const FSentSegmentInfoArray& ProcessedMessages)
+{
+	int32 BytesSent = 0;
 	for (const FSentSegmentInfo& Info : ProcessedMessages)
 	{
 		if ((Info.bFullySent && Info.bRequiresRequeue) ||
-			(Info.bRequiresRequeue && NodeInfo.WorkQueue.IsFull()) )
+			(Info.bRequiresRequeue && WorkQueue.IsFull()) )
 		{
 			NodeInfo.OverflowForPendingAck.Add(Info.MessageId);
 		}
 		else if (Info.bRequiresRequeue)
 		{
-			NodeInfo.WorkQueue.Enqueue(Info.MessageId);
+			WorkQueue.Enqueue(Info.MessageId);
 		}
 
 		BytesSent += Info.BytesSent;
 	}
+	return BytesSent;	
+}
 
-	if (CVarCheckForExpiredWithFullQueue.GetValueOnAnyThread() > 0 && !NodeInfo.CanSendSegments() && NodeInfo.WorkQueue.IsFull())
+bool FUdpMessageProcessor::ProcessQueue(FNodeInfo& NodeInfo, TUdpCircularQueue<int32>& WorkQueue, int64 DrainLimit, int32& BytesSent, bool bShouldDrainOverflowQueue)
+{
+	int64 SendCapacity = NodeInfo.SendCapacity();
+	int64 MaxSend = SendCapacity - DrainLimit;
+
+	while (!WorkQueue.IsEmpty() && NodeInfo.CanSendSegments() && SendCapacity > DrainLimit )
+	{
+		if (TOptional<FSentSegmentInfoArray> ProcessedMessages = DrainQueue(NodeInfo, WorkQueue, MaxSend, bShouldDrainOverflowQueue))
+		{
+			bShouldDrainOverflowQueue = false;
+			BytesSent += RequeueWork(NodeInfo, WorkQueue, *ProcessedMessages);
+			SendCapacity = NodeInfo.SendCapacity();
+			MaxSend = SendCapacity - DrainLimit;
+		}
+		else
+		{
+			// Socket Error
+			return false;
+		}
+	}
+	return true;
+}
+
+int32 FUdpMessageProcessor::UpdateSegmenters(FNodeInfo& NodeInfo)
+{
+	SCOPED_MESSAGING_TRACE(FUdpMessageProcessor_UpdateSegmenters);
+
+	const float ReliableQueuePriority = GetDefault<UUdpMessagingSettings>()->ReliableQueuePriority / 100.f;
+
+	int32 BytesSent = 0;
+	int64 SendCapacity = NodeInfo.SendCapacity();
+	const int64 SendRoomForUnreliables = (1 - ReliableQueuePriority) * SendCapacity;
+	int64 MaxSendForReliables = SendCapacity - SendRoomForUnreliables;
+
+	// Prioritize the reliable queue first. The reliable queue always gets more bandwidth.  We drain the overflow queue first and only
+	// on the first run as we don't want it to monopolize our time either.
+	// 
+	if (!ProcessQueue(NodeInfo, NodeInfo.ReliableWorkQueue, SendRoomForUnreliables, BytesSent, true))
+	{
+		// Socket error 
+		return -1;
+	}
+
+	if (!ProcessQueue(NodeInfo, NodeInfo.UnreliableWorkQueue, 0, BytesSent))
+	{
+		// Socket error 
+		return -1;
+	}
+
+	if (NodeInfo.SendCapacity() > 0)
+	{
+		// If there is still more room left go back to work on the reliable reliable queue until we have nothing left or run out
+		// of capacity.
+		if (!ProcessQueue(NodeInfo, NodeInfo.ReliableWorkQueue, 0, BytesSent))
+		{
+			// Socket error 
+			return -1;
+		}
+	}
+	
+	if (CVarCheckForExpiredWithFullQueue.GetValueOnAnyThread() > 0 && !NodeInfo.CanSendSegments() && NodeInfo.ReliableWorkQueue.IsFull())
 	{
 		UE_LOG(LogUdpMessaging, Warning, TEXT("Work queue is full sending to node at address %s. We cannot send new data. Attempting to expire old inflight segments."), *NodeInfo.Endpoint.ToString());
 		// RemoveLostSegments will clean-up the tracking map that we use to determine what is "inflight". Old data will be expired and scheduled
diff --git a/Engine/Plugins/Messaging/UdpMessaging/Source/UdpMessaging/Private/Transport/UdpMessageProcessor.h b/Engine/Plugins/Messaging/UdpMessaging/Source/UdpMessaging/Private/Transport/UdpMessageProcessor.h
--- a/Engine/Plugins/Messaging/UdpMessaging/Source/UdpMessaging/Private/Transport/UdpMessageProcessor.h
+++ b/Engine/Plugins/Messaging/UdpMessaging/Source/UdpMessaging/Private/Transport/UdpMessageProcessor.h
@@ -113,6 +113,8 @@ struct FSentSegmentInfo
 	}
 };
 
+using FSentSegmentInfoArray = TArray<FSentSegmentInfo, TInlineAllocator<64>>;
+
 namespace UE::Private::MessageProcessor
 {
 /**
@@ -155,8 +157,11 @@ class FUdpMessageProcessor
 		/** Holds the collection of message segmenters. */
 		TMap<int32, TSharedPtr<FUdpMessageSegmenter>> Segmenters;
 
-		/** Holds of queue of MessageIds to send. They are processed in round-robin fashion. */
-		TUdpCircularQueue<int32> WorkQueue;
+		/** Holds of queue of MessageIds to send for reliable messages. They are processed in round-robin fashion. */
+		TUdpCircularQueue<int32> ReliableWorkQueue;
+
+		/** Holds of queue of MessageIds to send for unreliable messages. They are processed in round-robin fashion. */
+		TUdpCircularQueue<int32> UnreliableWorkQueue;
 
 		/* Holds the list of messages that are still pending acknowledgement or overflow if overcommitted on work queue. */
 		TArray<int32> OverflowForPendingAck;
@@ -176,8 +181,8 @@ class FUdpMessageProcessor
 		/** Average time from when the packet was sent to ack received. */
 		FTimespan AvgRoundTripTime = FTimespan::FromMilliseconds(20);
 
-		/** Maximum number of segments we can have in flight. Start with a small value and let it grow based on performance. */
-		uint64 WindowSize = 64;
+		/** Maximum number of segments we can have in flight. Start with a middle of the road value and adjust based on network performance. */
+		uint64 WindowSize = 1024;
 
 		/** Various transport statistics for this endpoint */
 		FMessageTransportStatistics Statistics;
@@ -190,7 +195,7 @@ class FUdpMessageProcessor
 		 */
 		bool CanCommitToWorkQueue(const FDateTime& InCurrentTime)
 		{
-			const bool bWorkQueueFull = WorkQueue.IsFull();
+			const bool bWorkQueueFull = ReliableWorkQueue.IsFull();
 			if (bWorkQueueFull)
 			{
 				const double ReportingIntervalInSeconds = 1;
@@ -427,6 +432,11 @@ class FUdpMessageProcessor
 			return InflightSegments.Num() < WindowSize;
 		}
 
+		int64 SendCapacity() const
+		{
+			return static_cast<int64>(WindowSize) - static_cast<int64>(InflightSegments.Num());
+		}
+		
 		/** Resets the endpoint info. */
 		void ResetIfRestarted(const FGuid& NewNodeId)
 		{
@@ -848,7 +858,24 @@ protected:
 	virtual void Tick() override;
 
 private:
+	/**
+	 * Drain a given send queue to the drain limit specified on the parameter listed. Also drain the overflow bucket if requested. Overflow
+	 * may occasionally fill from pending acks or work queue overflowing.  
+	 */
+	TOptional<FSentSegmentInfoArray> DrainQueue(FNodeInfo& NodeInfo, TUdpCircularQueue<int32>& WorkQueue, int64 DrainCount, const bool bDrainOverflow);
 
+	/**
+	 * Refill the given queue with new work based on the processed messages. They will be put back in the queue to be processed later
+	 * in the next round to ensure fairness in the sending of data.
+	 */
+	int32 RequeueWork(FNodeInfo& NodeInfo, TUdpCircularQueue<int32>& WorkQueue, const FSentSegmentInfoArray& ProcessedMessages);
+
+	/**
+	 * Process a given queue.  This calls into DrainQueue and RequeueWork repeatedly until we reach the drain limit which is the minimum
+	 * value of of send capacity.  For example, if DrainLimit is 0 then we work until the send capacity is 0. 
+	 */
+	bool ProcessQueue(FNodeInfo& NodeInfo, TUdpCircularQueue<int32>& WorkQueue, int64 DrainLimit, int32& BytesSent, bool bShouldDrainOverflowQueue = false);
+		
 	/** Send our list of known nodes to the known nodes. */
 	void SendKnownNodesToKnownNodes();
 
diff --git a/Engine/Plugins/Messaging/UdpMessaging/Source/UdpMessaging/Public/Shared/UdpMessagingSettings.h b/Engine/Plugins/Messaging/UdpMessaging/Source/UdpMessaging/Public/Shared/UdpMessagingSettings.h
--- a/Engine/Plugins/Messaging/UdpMessaging/Source/UdpMessaging/Public/Shared/UdpMessagingSettings.h
+++ b/Engine/Plugins/Messaging/UdpMessaging/Source/UdpMessaging/Public/Shared/UdpMessagingSettings.h
@@ -89,6 +89,13 @@ public:
 	UPROPERTY(config, EditAnywhere, Category = Transport, meta = (ClampMin="1024"), AdvancedDisplay)
 	uint16 WorkQueueSize = 1024;
 
+	/**
+	 * Reliable queue priority specified as a percentage with a range from 5-95%. Setting values above or below
+	 * can completely starve unreliable messages.
+	 */
+	UPROPERTY(config, EditAnywhere, Category = Transport, meta = (ClampMin="5", ClampMax="95"), AdvancedDisplay)
+	float ReliableQueuePriority = 75.0f;
+
 	/** Whether to stop the transport service when the application deactivates, and restart it when the application is reactivated */
 	UPROPERTY(config)
 	bool bStopServiceWhenAppDeactivates = true;
@@ -158,6 +165,14 @@ public:
 	UPROPERTY(config, EditAnywhere, Category=Transport, meta = (ClampMax="32"), AdvancedDisplay)
 	uint32 MaxConcurrentDeserializationTasks = 8;
 
+	/**
+	 * The period of time to consider an active connection to be timed-out. A regular heart beat is expected from connected endpoints and this will
+	 * expire that connection if no data is received after a period of time.  If zero is specifed the endpoints are never timed out (this should
+	 * only be use for debugging purposes).
+	 */
+	UPROPERTY(config, EditAnywhere, Category=Transport, meta = (ClampMin="0"), AdvancedDisplay)
+	float ConnectionTimeoutPeriod = 5.0f;
+
 public:
 
 	/** Whether the UDP tunnel is enabled. */
