The loopback device is not added to the multicast group by default.  When using IF_ANY make sure we include the loopback on the multicast group.

[RN] In UDP messaging, the loopback device was not added to the multicast group; so only the first editor instance launch would receive multicast packets.  We now always join the group across all instances so other editor instances can be in the same multicast group.

diff --git a/Engine/Plugins/Messaging/UdpMessaging/Source/UdpMessaging/Private/Transport/UdpMessageTransport.cpp b/Engine/Plugins/Messaging/UdpMessaging/Source/UdpMessaging/Private/Transport/UdpMessageTransport.cpp
--- a/Engine/Plugins/Messaging/UdpMessaging/Source/UdpMessaging/Private/Transport/UdpMessageTransport.cpp
+++ b/Engine/Plugins/Messaging/UdpMessaging/Source/UdpMessaging/Private/Transport/UdpMessageTransport.cpp
@@ -17,7 +17,6 @@
 #include "Transport/UdpMessageProcessor.h"
 #include "UdpMessageProcessor.h"
 
-
 /* FUdpMessageTransport structors
  *****************************************************************************/
 
@@ -223,25 +222,32 @@ void DoJoinMulticastGroup(const TSharedRef<FInternetAddr>& MulticastAddr, const
 	}
 }
 
-
 void JoinedToGroup(const FIPv4Endpoint& UnicastEndpoint, const FIPv4Endpoint& MulticastEndpoint, FSocket* MulticastSocket)
 {
 #if PLATFORM_SUPPORTS_UDP_MULTICAST_GROUP
 	TSharedRef<FInternetAddr> MulticastAddr = MulticastEndpoint.ToInternetAddr();
 	if (UnicastEndpoint.Address == FIPv4Address::Any)
 	{
-		TArray<TSharedPtr<FInternetAddr>> LocapIps;
-		ISocketSubsystem::Get(PLATFORM_SOCKETSUBSYSTEM)->GetLocalAdapterAddresses(LocapIps);
-		for (const TSharedPtr<FInternetAddr>& LocalIp : LocapIps)
+		TArray<TSharedPtr<FInternetAddr>> LocalAdapterAddresses;
+		ISocketSubsystem::Get(PLATFORM_SOCKETSUBSYSTEM)->GetLocalAdapterAddresses(LocalAdapterAddresses);
+
+		const bool bShouldBindToLocalHostAddr = LocalAdapterAddresses.Num() == 0;
+		
+		// The loopback address is not included in the call to GetLocalAdapterAddresses. Make sure we add it so that loopback address is included in the group.
+		TSharedRef<FInternetAddr> Loopback = ISocketSubsystem::Get(PLATFORM_SOCKETSUBSYSTEM)->CreateInternetAddr();
+		Loopback->SetLoopbackAddress();
+		LocalAdapterAddresses.Add(Loopback.ToSharedPtr());
+
+		// Now iterate over all addresss and join the multicast group.
+		for (const TSharedPtr<FInternetAddr>& LocalIp : LocalAdapterAddresses)
 		{
 			DoJoinMulticastGroup(MulticastAddr, LocalIp, MulticastSocket);
 		}
 
-		// GetLocalAdapterAddresses returns empty list when all network adapters are offline
-		if (LocapIps.Num() == 0)
+		if (bShouldBindToLocalHostAddr)
 		{
 			bool bCanBindAll = false;
-			DoJoinMulticastGroup(MulticastAddr, ISocketSubsystem::Get(PLATFORM_SOCKETSUBSYSTEM)->GetLocalHostAddr(*GLog, bCanBindAll), MulticastSocket);
+			DoJoinMulticastGroup(MulticastAddr, ISocketSubsystem::Get(PLATFORM_SOCKETSUBSYSTEM)->GetLocalHostAddr(*GLog, bCanBindAll), MulticastSocket);			
 		}
 	}
 	else
@@ -279,7 +285,7 @@ bool FUdpMessageTransport::StartTransport(IMessageTransportHandler& Handler)
 		UE_LOG(LogUdpMessaging, Display, TEXT("Unicast socket bound to '%s:%d'."), *UnicastEndpoint.Address.ToString(), PortNo);
 	}
 #endif
-
+	
 	// create & initialize multicast socket (optional)
 	MulticastSocket = FUdpSocketBuilder(TEXT("UdpMessageMulticastSocket"))
 		.AsNonBlocking()
@@ -290,7 +296,7 @@ bool FUdpMessageTransport::StartTransport(IMessageTransportHandler& Handler)
 		// The instance that receives the packet may vary over time, seemingly based on the congestion of its socket.
 		// This isn't the intended usage.
 		// To allow traffic to be sent directly to unicast for discovery, set the interface and port for the unicast endpoint
-		// However for legacy reason, keep binding this as well, although it might be unreliable in some cases
+		// However for legacy reason, keep binding this as well, although it might be unreliable in some cases. 
 		.BoundToAddress(UnicastEndpoint.Address)
 #endif
 		.BoundToPort(MulticastEndpoint.Port)
@@ -350,10 +356,14 @@ bool FUdpMessageTransport::StartTransport(IMessageTransportHandler& Handler)
 	if (MulticastSocket != nullptr)
 	{
 		MulticastReceiver = new FUdpSocketReceiver(MulticastSocket, ThreadWaitTime, TEXT("UdpMessageMulticastReceiver"));
-		MulticastReceiver->OnDataReceived().BindRaw(this, &FUdpMessageTransport::HandleSocketDataReceived);
+		MulticastReceiver->OnDataReceived().BindRaw(this, &FUdpMessageTransport::HandleMulticastSocketDataReceived);
 		MulticastReceiver->SetMaxReadBufferSize(2048);
 		MulticastReceiver->Start();
 	}
+	else
+	{
+		UE_LOG(LogUdpMessaging, Warning, TEXT("Multicast socket was null so there is no multicast receiver."));
+	}
 
 #if PLATFORM_DESKTOP
 	UnicastReceiver = new FUdpSocketReceiver(UnicastSocket, ThreadWaitTime, TEXT("UdpMessageUnicastReceiver"));
@@ -678,6 +688,14 @@ bool FUdpMessageTransport::RestartTransport()
 	return StartTransport(*Handler);
 }
 
+void FUdpMessageTransport::HandleMulticastSocketDataReceived(const TSharedPtr<FArrayReader, ESPMode::ThreadSafe>& Data, const FIPv4Endpoint& Sender)
+{
+	if (MessageProcessor != nullptr)
+	{
+		MessageProcessor->EnqueueInboundSegment(Data, Sender);
+	}
+}
+
 
 void FUdpMessageTransport::HandleSocketDataReceived(const TSharedPtr<FArrayReader, ESPMode::ThreadSafe>& Data, const FIPv4Endpoint& Sender)
 {
diff --git a/Engine/Plugins/Messaging/UdpMessaging/Source/UdpMessaging/Private/Transport/UdpMessageTransport.h b/Engine/Plugins/Messaging/UdpMessaging/Source/UdpMessaging/Private/Transport/UdpMessageTransport.h
--- a/Engine/Plugins/Messaging/UdpMessaging/Source/UdpMessaging/Private/Transport/UdpMessageTransport.h
+++ b/Engine/Plugins/Messaging/UdpMessaging/Source/UdpMessaging/Private/Transport/UdpMessageTransport.h
@@ -138,6 +138,9 @@ private:
 	/** Handles received socket data. */
 	void HandleSocketDataReceived(const TSharedPtr<FArrayReader, ESPMode::ThreadSafe>& Data, const FIPv4Endpoint& Sender);
 
+	/** Handles received socket data for multicast receiver. */
+	void HandleMulticastSocketDataReceived(const TSharedPtr<FArrayReader, ESPMode::ThreadSafe>& Data, const FIPv4Endpoint& Sender);
+		
 private:
 	/** Clear the deny list */
 	void DoClearDenyCandidateList();
