Adjust UDP Messaging sender to fully utilize the available capacity in the send pipe and prioritize reliable messages over unreliable ones.   This is a resubmit after the first try to implement this.  There was a race condition that prevented Multi-user from joining during the initial autoconnect in a -game node.

We also add settings for the prioritized queue and the time out period for known nodes.  The timeout period was fixed at 5 seconds and this gives users more control over its value.

[RN] Better utilization and prioritization of data sent over UDP messaging when reliable and unreliable data is sharing the transmission window.

diff --git a/Engine/Plugins/Messaging/UdpMessaging/Source/UdpMessaging/Private/Transport/UdpMessageBeacon.cpp b/Engine/Plugins/Messaging/UdpMessaging/Source/UdpMessaging/Private/Transport/UdpMessageBeacon.cpp
--- a/Engine/Plugins/Messaging/UdpMessaging/Source/UdpMessaging/Private/Transport/UdpMessageBeacon.cpp
+++ b/Engine/Plugins/Messaging/UdpMessaging/Source/UdpMessaging/Private/Transport/UdpMessageBeacon.cpp
@@ -1,8 +1,10 @@
 // Copyright Epic Games, Inc. All Rights Reserved.
 
+
 #include "Transport/UdpMessageBeacon.h"
 #include "HAL/PlatformProcess.h"
 #include "UdpMessagingPrivate.h"
+#include "Common/UdpSocketSender.h"
 
 #include "HAL/Event.h"
 #include "HAL/RunnableThread.h"
@@ -20,16 +22,17 @@ const FTimespan FUdpMessageBeacon::MinimumInterval = FTimespan::FromMilliseconds
 /* FUdpMessageHelloSender structors
  *****************************************************************************/
 
-FUdpMessageBeacon::FUdpMessageBeacon(FSocket* InSocket, const FGuid& InSocketId, const FIPv4Endpoint& InMulticastEndpoint)
+FUdpMessageBeacon::FUdpMessageBeacon(FUdpSocketSender* InSocketSender, const FGuid& InSocketId, const FIPv4Endpoint& InMulticastEndpoint)
 	: BeaconInterval(MinimumInterval)
-	, LastEndpointCount(1)
+	, LastEndpointCount(0)
 	, LastHelloSent(FDateTime::MinValue())
 	, NextHelloTime(FDateTime::UtcNow())
 	, NodeId(InSocketId)
-	, Socket(InSocket)
+	, SocketSender(InSocketSender)
 	, Stopping(false)
 	, bSocketError(false)
 {
+	check(SocketSender);
 	EndpointLeftEvent = FPlatformProcess::GetSynchEventFromPool(false);
 	MulticastAddress = InMulticastEndpoint.ToInternetAddr();
 	AddLocalEndpoints();
@@ -142,23 +145,15 @@ bool FUdpMessageBeacon::SendSegment(EUdpMessageSegments SegmentType, const FTime
 		Header.SegmentType = SegmentType;
 	}
 
-	FArrayWriter Writer;
-	{
-		Writer << Header;
-		Writer << NodeId;
-	}
-
-	int32 Sent;
+	TSharedRef<FArrayWriter, ESPMode::ThreadSafe> Writer = MakeShared<FArrayWriter, ESPMode::ThreadSafe>();
 
-	if (!Socket->Wait(ESocketWaitConditions::WaitForWrite, SocketWaitTime))
-	{
-		return false; // socket not ready for sending
-	}
+	*Writer << Header;
+	*Writer << NodeId;
 
-	if (!Socket->SendTo(Writer.GetData(), Writer.Num(), Sent, *MulticastAddress))
+	if (!SocketSender->Send(Writer, MulticastAddress))
 	{
 		bSocketError = true;
-		return false; // send failed
+		return false;
 	}
 
 	return true;
@@ -176,28 +171,21 @@ bool FUdpMessageBeacon::SendPing(const FTimespan& SocketWaitTime)
 	}
 	uint8 ActualProtocolVersion = UDP_MESSAGING_TRANSPORT_PROTOCOL_VERSION;
 
-	FArrayWriter Writer;
+	TSharedRef<FArrayWriter, ESPMode::ThreadSafe> Writer = MakeShared<FArrayWriter, ESPMode::ThreadSafe>();
 	{
-		Writer << Header;
-		Writer << NodeId;
-		Writer << ActualProtocolVersion; // Send our actual Protocol version as part of the ping message
+		*Writer << Header;
+		*Writer << NodeId;
+		*Writer << ActualProtocolVersion; // Send our actual Protocol version as part of the ping message
 	}
 
 
-	if (!Socket->Wait(ESocketWaitConditions::WaitForWrite, SocketWaitTime))
-	{
-		return false; // socket not ready for sending
-	}
-
-	int32 Sent;
 	bool Result = true;
 	for (const auto& StaticAddress : StaticAddresses)
 	{
-		if (!Socket->SendTo(Writer.GetData(), Writer.Num(), Sent, *StaticAddress))
+		if (!SocketSender->Send(Writer, StaticAddress))
 		{	
 			Result = false; // send failed
 		}
-
 	}
 	return Result;
 }
@@ -213,7 +201,6 @@ void FUdpMessageBeacon::Update(const FDateTime& CurrentTime, const FTimespan& So
 	}
 
 	BeaconInterval = FMath::Max(MinimumInterval, IntervalPerEndpoint * LastEndpointCount);
-
 	if (SendSegment(EUdpMessageSegments::Hello, SocketWaitTime))
 	{
 		NextHelloTime = CurrentTime + BeaconInterval;
diff --git a/Engine/Plugins/Messaging/UdpMessaging/Source/UdpMessaging/Private/Transport/UdpMessageBeacon.h b/Engine/Plugins/Messaging/UdpMessaging/Source/UdpMessaging/Private/Transport/UdpMessageBeacon.h
--- a/Engine/Plugins/Messaging/UdpMessaging/Source/UdpMessaging/Private/Transport/UdpMessageBeacon.h
+++ b/Engine/Plugins/Messaging/UdpMessaging/Source/UdpMessaging/Private/Transport/UdpMessageBeacon.h
@@ -14,6 +14,7 @@ enum class EUdpMessageSegments : uint8;
 class FEvent;
 class FInternetAddr;
 class FSocket;
+class FUdpSocketSender;
 
 /**
  * Implements a beacon sender thread.
@@ -28,12 +29,12 @@ public:
 	/** 
 	 * Creates and initializes a new beacon sender.
 	 *
-	 * @param InSocket The network socket used to send Hello segments.
+	 * @param InSocket The network socket sender used to send Hello segments.
 	 * @param InSocketId The network socket identifier (used to detect unicast endpoint).
 	 * @param InMulticastEndpoint The multicast group endpoint to transport messages to.
 	 * @param InStaticEndpoints The static nodes to broadcast to alongside the multicast.
 	 */
-	FUdpMessageBeacon(FSocket* InSocket, const FGuid& InSocketId, const FIPv4Endpoint& InMulticastEndpoint);
+	FUdpMessageBeacon(FUdpSocketSender* InSocketSender, const FGuid& InSocketId, const FIPv4Endpoint& InMulticastEndpoint);
 
 	/** Virtual destructor. */
 	virtual ~FUdpMessageBeacon();
@@ -157,8 +158,8 @@ private:
 	/** Holds local node identifier. */
 	FGuid NodeId;
 
-	/** Holds the socket used to send Hello segments. */
-	FSocket* Socket;
+	/** Holds the socket sender used to send Hello segments. */
+	FUdpSocketSender* SocketSender;
 
 	/** Holds a flag indicating that the thread is stopping. */
 	bool Stopping;
diff --git a/Engine/Plugins/Messaging/UdpMessaging/Source/UdpMessaging/Private/Transport/UdpMessageProcessor.cpp b/Engine/Plugins/Messaging/UdpMessaging/Source/UdpMessaging/Private/Transport/UdpMessageProcessor.cpp
--- a/Engine/Plugins/Messaging/UdpMessaging/Source/UdpMessaging/Private/Transport/UdpMessageProcessor.cpp
+++ b/Engine/Plugins/Messaging/UdpMessaging/Source/UdpMessaging/Private/Transport/UdpMessageProcessor.cpp
@@ -11,6 +11,7 @@
 
 #include "Transport/UdpDeserializedMessage.h"
 #include "UdpCircularQueue.h"
+#include "UdpMessagingPrivate.h"
 #include "UdpMessagingTracing.h"
 #include "Shared/UdpMessagingSettings.h"
 #include "Transport/UdpMessageBeacon.h"
@@ -103,7 +104,8 @@ FUdpMessageProcessor::FNodeInfo::FNodeInfo()
 	: LastSegmentReceivedTime(FDateTime::MinValue())
 	, NodeId()
 	, ProtocolVersion(UDP_MESSAGING_TRANSPORT_PROTOCOL_VERSION)
-	, WorkQueue(UE::Private::MessageProcessor::GetMessageProcessorWorkQueueSize())
+	, ReliableWorkQueue(UE::Private::MessageProcessor::GetMessageProcessorWorkQueueSize())
+	, UnreliableWorkQueue(UE::Private::MessageProcessor::GetMessageProcessorWorkQueueSize())
 {
 	ComputeWindowSize(0,0);
 }
@@ -151,6 +153,8 @@ FUdpMessageProcessor::~FUdpMessageProcessor()
 
 	Thread = {};
 	Beacon = {};
+
+	// NOTE: Socket sender must be destroyed after the beacon because the beacon uses the socket sender.
 	SocketSender = {};
 
 	IPluginManager::Get().OnLoadingPhaseComplete().RemoveAll(this);
@@ -309,7 +313,10 @@ void FUdpMessageProcessor::SendSegmenterStatsToListeners(int32 MessageId, FGuid
 			// Convert segment data into bytes.
 			UDP_MESSAGING_SEGMENT_SIZE * SegmentCount,
 			UDP_MESSAGING_SEGMENT_SIZE * (SegmentCount - Segmenter->GetPendingSendSegmentsCount()),
-			UDP_MESSAGING_SEGMENT_SIZE * (Segmenter->GetAcknowledgedSegmentsCount())
+			UDP_MESSAGING_SEGMENT_SIZE * (Segmenter->GetAcknowledgedSegmentsCount()),
+			EnumHasAnyFlags(Segmenter->GetMessageFlags(), EMessageFlags::Reliable),
+			EnumHasAnyFlags(Segmenter->GetMessageFlags(), EMessageFlags::Traced),
+			Segmenter->AreAcknowledgementsComplete()
 		});
 }
 
@@ -336,9 +343,10 @@ bool FUdpMessageProcessor::Init()
 {
 	if (!bIsInitialized)
 	{
-		Beacon = MakeUnique<FUdpMessageBeacon>(Socket, LocalNodeId, MulticastEndpoint);
 		SocketSender = MakeUnique<FUdpSocketSender>(Socket, TEXT("FUdpMessageProcessor.Sender"));
 
+		Beacon = MakeUnique<FUdpMessageBeacon>(SocketSender.Get(), LocalNodeId, MulticastEndpoint);
+
 		// Current protocol version 17
 		SupportedProtocolVersions.Add(UDP_MESSAGING_TRANSPORT_PROTOCOL_VERSION);
 		// Support Protocol version 10, 11, 12, 13, 14, 15, 16
@@ -452,14 +460,15 @@ void FUdpMessageProcessor::AcknowledgeReceipt(int32 MessageId, const FNodeInfo&
 		AcknowledgeChunk.MessageId = MessageId;
 	}
 
-	FArrayWriter Writer;
+	TSharedRef<FArrayWriter, ESPMode::ThreadSafe> Writer = MakeShared<FArrayWriter, ESPMode::ThreadSafe>();
+	*Writer << Header;
+	AcknowledgeChunk.Serialize(*Writer, NodeInfo.ProtocolVersion);
+
+	if (!SocketSender->Send(Writer, NodeInfo.Endpoint))
 	{
-		Writer << Header;
-		AcknowledgeChunk.Serialize(Writer, NodeInfo.ProtocolVersion);
+		UE_LOG(LogUdpMessaging, Warning, TEXT("FUdpMessageProcessor::AcknowledgeReceipt failed to send %s."), *NodeInfo.Endpoint.ToString());
+		return;
 	}
-
-	int32 OutSent;
-	Socket->SendTo(Writer.GetData(), Writer.Num(), OutSent, *NodeInfo.Endpoint.ToInternetAddr());
 	UE_LOG(LogUdpMessaging, Verbose, TEXT("Sending EUdpMessageSegments::Acknowledge for msg %d from %s"), MessageId, *NodeInfo.NodeId.ToString());
 }
 
@@ -489,10 +498,17 @@ void FUdpMessageProcessor::ConsumeInboundSegments()
 			continue;
 		}
 
+		
 		FUdpMessageSegment::FHeader Header;
 		*Segment.Data << Header;
 
-		if (FilterSegment(Header, Segment.Sender))
+		const bool bPassesFilter = FilterSegment(Header, Segment.Sender);
+		if (bEnableMessageDelegates)
+		{
+			InboundSegmentDelegate.Broadcast(Segment, Header.SenderNodeId, Header.SegmentType, bPassesFilter);
+		}
+
+		if (bPassesFilter)
 		{
 			FNodeInfo& NodeInfo = KnownNodes.FindOrAdd(Header.SenderNodeId);
 
@@ -582,6 +598,7 @@ bool FUdpMessageProcessor::ConsumeOneOutboundMessage(const FOutboundMessage& Out
 					  [FindNodeWithLog](const FGuid &Id) {return FindNodeWithLog(Id);},
 					  [this](const FGuid &Id) {return KnownNodes.Find(Id);} );
 
+	
 	const bool bCanConsume = Algo::AllOf(Recipients, [this](FNodeInfo *Node)
 		{
 			return Node->CanCommitToWorkQueue(CurrentTime);
@@ -599,8 +616,9 @@ bool FUdpMessageProcessor::ConsumeOneOutboundMessage(const FOutboundMessage& Out
 		}
 		for (FNodeInfo *RecipientNodeInfo : Recipients)
 		{
-			UE_LOG(LogUdpMessaging, Verbose, TEXT("Passing %" INT64_FMT " byte message to be segement-sent to %s with id %s"),
+			UE_LOG(LogUdpMessaging, Verbose, TEXT("Passing %" INT64_FMT " byte message to be segment-sent %d to %s with id %s"),
 				   OutboundMessage.SerializedMessage->TotalSize(),
+				   LastSentMessage,
 				   *RecipientNodeInfo->Endpoint.ToString(),
 				   *RecipientNodeInfo->NodeId.ToString());
 
@@ -609,12 +627,18 @@ bool FUdpMessageProcessor::ConsumeOneOutboundMessage(const FOutboundMessage& Out
 				// Discard unreliable messages that cannot be sent.
 				continue;
 			}
+
 			TSharedPtr<FUdpMessageSegmenter> Segmenter = MakeShared<FUdpMessageSegmenter>(OutboundMessage.SerializedMessage.ToSharedRef(), UDP_MESSAGING_SEGMENT_SIZE);
 			RecipientNodeInfo->Segmenters.Add(
 				LastSentMessage, Segmenter);
-
-			RecipientNodeInfo->WorkQueue.Enqueue(LastSentMessage);
-			SendSegmenterStatsToListeners(LastSentMessage, RecipientNodeInfo->NodeId, Segmenter);
+			if (bIsReliable)
+			{
+				RecipientNodeInfo->ReliableWorkQueue.Enqueue(LastSentMessage);
+			}
+			else
+			{
+				RecipientNodeInfo->UnreliableWorkQueue.Enqueue(LastSentMessage);
+			}
 		}
 	}
 
@@ -710,7 +734,6 @@ void FUdpMessageProcessor::ProcessAcknowledgeSegmentsSegment(FInboundSegment& Se
 	AcknowledgeChunk.Serialize(*Segment.Data, NodeInfo.ProtocolVersion);
 
 	UE_LOG(LogUdpMessaging, Verbose, TEXT("Received AcknowledgeSegments for %d from %s"), AcknowledgeChunk.MessageId, *NodeInfo.NodeId.ToString());
-
 	NodeInfo.MarkAcks(AcknowledgeChunk.MessageId, AcknowledgeChunk.Segments, CurrentTime);
 }
 
@@ -764,11 +787,16 @@ void FUdpMessageProcessor::ProcessDataSegment(FInboundSegment& Segment, FNodeInf
 	if(ReassemblerUpdated.IsBound())
 	{
 		ReassemblerUpdated.Broadcast(
-			{NodeInfo.NodeId,
-			 DataChunk.MessageId,
-			 UDP_MESSAGING_SEGMENT_SIZE * ReassembledMessage->GetPendingSegmentsCount(),
-			 ReassembledMessage->GetReceivedBytes()}
-			);
+			{
+				NodeInfo.NodeId,
+				DataChunk.MessageId,
+				UDP_MESSAGING_SEGMENT_SIZE * ReassembledMessage->GetPendingSegmentsCount(),
+				ReassembledMessage->GetReceivedBytes(),
+				ReassembledMessage->IsReliable(),
+				EnumHasAnyFlags(ReassembledMessage->GetFlags(), EMessageFlags::Traced),
+				ReassembledMessage->IsComplete()
+			}
+		);
 	}
 
 	// Deliver or re-sequence message
@@ -835,14 +863,16 @@ void FUdpMessageProcessor::ProcessPingSegment(FInboundSegment& Segment, FNodeInf
 		Header.SegmentType = EUdpMessageSegments::Pong;
 	}
 
-	FArrayWriter Writer;
+	TSharedRef<FArrayWriter, ESPMode::ThreadSafe> Writer = MakeShared<FArrayWriter, ESPMode::ThreadSafe>();
 	{
-		Writer << Header;
-		Writer << LocalNodeId;
+		*Writer << Header;
+		*Writer << LocalNodeId;
 	}
 
-	int32 OutSent;
-	Socket->SendTo(Writer.GetData(), Writer.Num(), OutSent, *NodeInfo.Endpoint.ToInternetAddr());
+	if (!SocketSender->Send(Writer, NodeInfo.Endpoint))
+	{
+		UE_LOG(LogUdpMessaging, Warning, TEXT("FUdpMessageProcessor::ProcessPingSegment failed to send ping segment %s."), *NodeInfo.Endpoint.ToString());		
+	}
 }
 
 
@@ -966,7 +996,7 @@ int32 GetMaxSendRate()
 
 bool HasTimedOut(const FDateTime& ReceivedTime, const FDateTime& CurrentTime)
 {
-	const int32 TimeoutInSeconds = 5;
+	const float TimeoutInSeconds = GetDefault<UUdpMessagingSettings>()->ConnectionTimeoutPeriod;
 	FTimespan TimeoutInterval = TimeoutInSeconds * FTimespan::FromMilliseconds(1000);
 	return (ReceivedTime + TimeoutInterval) < CurrentTime;
 }
@@ -975,6 +1005,12 @@ void FUdpMessageProcessor::RemoveDeadNodes()
 {
 	// Remove dead nodes
 	FTimespan DeadHelloTimespan = DeadHelloIntervals * Beacon->GetBeaconInterval();
+	if (DeadHelloTimespan.IsZero())
+	{
+		// We never expire connections.
+		return;
+	}
+
 	for (auto It = KnownNodes.CreateIterator(); It; ++It)
 	{
 		FGuid& NodeId = It->Key;
@@ -982,7 +1018,7 @@ void FUdpMessageProcessor::RemoveDeadNodes()
 
 		if ((NodeId.IsValid()) && HasTimedOut(NodeInfo.LastSegmentReceivedTime,CurrentTime))
 		{
-			UE_LOG(LogUdpMessaging, Verbose, TEXT("FUdpMessageProcessor::UpdateKnownNodes: Removing Node %s (%s)"), *NodeInfo.NodeId.ToString(), *NodeInfo.Endpoint.ToString());
+			UE_LOG(LogUdpMessaging, Display, TEXT("FUdpMessageProcessor::UpdateKnownNodes: Removing Node %s (%s)"), *NodeInfo.NodeId.ToString(), *NodeInfo.Endpoint.ToString());
 			NodeLostDelegate.ExecuteIfBound(NodeId);
 			It.RemoveCurrent();
 		}
@@ -1003,8 +1039,7 @@ bool FUdpMessageProcessor::MoreToSend()
 
 void FUdpMessageProcessor::HandleSocketError(const FNodeInfo& NodeInfo) const
 {
-	UE_LOG(LogUdpMessaging, Error,
-		   TEXT("Socket error detected when communicating with %s. Banning communication to that endpoint."), *NodeInfo.Endpoint.ToString());
+	UE_LOG(LogUdpMessaging, Error,  TEXT("Socket error detected when communicating with %s. Banning communication to that endpoint."), *NodeInfo.Endpoint.ToString());
 	ErrorSendingToEndpointDelegate.Execute(NodeInfo.NodeId, NodeInfo.Endpoint);
 }
 
@@ -1050,18 +1085,17 @@ void FUdpMessageProcessor::SendKnownNodesToKnownNodes()
 
 		OutEndpoints = SlicedEndpointView;
 		OutIds = SlicedIdView;
-		FArrayWriter Writer;
-		{
-			Writer << Header;
-			Writer << OutEndpoints;
-			Writer << OutIds;
-		}
 		UE_LOG(LogUdpMessaging, VeryVerbose, TEXT("FUdpMessageProcessor::SendKnownNodesToKnownNodes Sending updated known nodes."));
 
 		for (const FIPv4Endpoint& Endpoint : Endpoints)
 		{
-			int32 Sent;
-			if (!Socket->SendTo(Writer.GetData(), Writer.Num(), Sent, *Endpoint.ToInternetAddr()))
+			TSharedRef<FArrayWriter, ESPMode::ThreadSafe> Writer = MakeShared<FArrayWriter, ESPMode::ThreadSafe>();
+			{
+				*Writer << Header;
+				*Writer << OutEndpoints;
+				*Writer << OutIds;
+			}
+			if (!SocketSender->Send(Writer, Endpoint))
 			{
 				UE_LOG(LogUdpMessaging, Warning, TEXT("FUdpMessageProcessor::SendKnownNodesToKnownNodes failed to share endpoint information to %s."), *Endpoint.ToString());
 			}
@@ -1254,6 +1288,7 @@ FSentSegmentInfo FUdpMessageProcessor::SendNextSegmentForMessageId(FNodeInfo& No
 		SentInfo.bIsReliable = EnumHasAnyFlags(Segmenter->GetMessageFlags(), EMessageFlags::Reliable);
 		SentInfo.bRequiresRequeue = !Segmenter->IsSendingComplete() || SentInfo.bIsReliable;
 		SentInfo.bFullySent = Segmenter->IsSendingComplete();
+		SentInfo.bIsTraced = EnumHasAnyFlags(Segmenter->GetMessageFlags(), EMessageFlags::Traced);
 
 		Segmenter->UpdateSentTime(CurrentTime);
 
@@ -1272,8 +1307,10 @@ FSentSegmentInfo FUdpMessageProcessor::SendNextSegmentForMessageId(FNodeInfo& No
 	}
 	else
 	{
+		SentInfo.bIsTraced = EnumHasAnyFlags(Segmenter->GetMessageFlags(), EMessageFlags::Traced);
+		SentInfo.bIsReliable = EnumHasAnyFlags(Segmenter->GetMessageFlags(), EMessageFlags::Reliable);
 		SentInfo.bRequiresRequeue = true;
-		SentInfo.bFullySent = Segmenter->IsSendingComplete();
+		SentInfo.bFullySent = Segmenter->IsInitialized() && Segmenter->IsSendingComplete();
 	}
 
 	NodeInfo.Statistics.TotalBytesSent += SentInfo.BytesSent;
@@ -1281,9 +1318,8 @@ FSentSegmentInfo FUdpMessageProcessor::SendNextSegmentForMessageId(FNodeInfo& No
 	return MoveTemp(SentInfo);
 }
 
-int32 FUdpMessageProcessor::UpdateSegmenters(FNodeInfo& NodeInfo)
+TOptional<FSentSegmentInfoArray> FUdpMessageProcessor::DrainQueue(FNodeInfo& NodeInfo, TUdpCircularQueue<int32>& WorkQueue, int64 DrainCount, const bool bDrainOverflow)
 {
-	SCOPED_MESSAGING_TRACE(FUdpMessageProcessor_UpdateSegmenters);
 	FUdpMessageSegment::FHeader Header
 	{
 		NodeInfo.ProtocolVersion,		// Header.ProtocolVersion - Send data segment using the node protocol version
@@ -1292,55 +1328,119 @@ int32 FUdpMessageProcessor::UpdateSegmenters(FNodeInfo& NodeInfo)
 		EUdpMessageSegments::Data		// Header.SegmentType
 	};
 
-	int32 MessageId;
-	int32 BytesSent = 0;
-	TArray<FSentSegmentInfo,TInlineAllocator<64>> ProcessedMessages;
-	auto SendOne = [this, &NodeInfo, &Header, &ProcessedMessages](int32 MessageId)
-	{
-		FSentSegmentInfo Info = SendNextSegmentForMessageId(NodeInfo, Header, MessageId);
-		ProcessedMessages.Emplace(MoveTemp(Info));
-	};
-	auto SocketError = [&ProcessedMessages]()
-	{
-		return ProcessedMessages.Top().bSendSocketError;
-	};
+	FSentSegmentInfoArray ProcessedMessages;
 
-	// Check for acknowledged segments first.
-	while (NodeInfo.CanSendSegments() && !NodeInfo.OverflowForPendingAck.IsEmpty())
+	if (bDrainOverflow)
 	{
-		SendOne(NodeInfo.OverflowForPendingAck.Pop());
-		if (SocketError())
+		// Check for acknowledged segments first.
+		while (NodeInfo.CanSendSegments() && !NodeInfo.OverflowForPendingAck.IsEmpty())
 		{
-			return -1;
+			FSentSegmentInfo Info = SendNextSegmentForMessageId(NodeInfo, Header, NodeInfo.OverflowForPendingAck.Pop());
+			if (Info.bSendSocketError)
+			{
+				return {};
+			}
+			ProcessedMessages.Emplace(MoveTemp(Info));
+			DrainCount--;
 		}
 	}
-
+	
+	int32 MessageId;
 	// Process messages in the work queue.
-	while (NodeInfo.CanSendSegments() && NodeInfo.WorkQueue.Dequeue(MessageId))
+	while (DrainCount > 0 && NodeInfo.CanSendSegments() && WorkQueue.Dequeue(MessageId))
 	{
-		SendOne(MessageId);
-		if (SocketError())
+		FSentSegmentInfo Info = SendNextSegmentForMessageId(NodeInfo, Header, MessageId);
+		if (Info.bSendSocketError)
 		{
-			return -1;
+			return {};
 		}
+		ProcessedMessages.Emplace(MoveTemp(Info));
+		DrainCount--;
 	}
+	return ProcessedMessages;
+}
 
+int32 FUdpMessageProcessor::RequeueWork(FNodeInfo& NodeInfo, TUdpCircularQueue<int32>& WorkQueue, const FSentSegmentInfoArray& ProcessedMessages)
+{
+	int32 BytesSent = 0;
 	for (const FSentSegmentInfo& Info : ProcessedMessages)
 	{
 		if ((Info.bFullySent && Info.bRequiresRequeue) ||
-			(Info.bRequiresRequeue && NodeInfo.WorkQueue.IsFull()) )
+			(Info.bRequiresRequeue && WorkQueue.IsFull()) )
 		{
 			NodeInfo.OverflowForPendingAck.Add(Info.MessageId);
 		}
 		else if (Info.bRequiresRequeue)
 		{
-			NodeInfo.WorkQueue.Enqueue(Info.MessageId);
+			WorkQueue.Enqueue(Info.MessageId);
 		}
 
 		BytesSent += Info.BytesSent;
 	}
+	return BytesSent;	
+}
+
+bool FUdpMessageProcessor::ProcessQueue(FNodeInfo& NodeInfo, TUdpCircularQueue<int32>& WorkQueue, int64 DrainLimit, int32& BytesSent, bool bShouldDrainOverflowQueue)
+{
+	int64 SendCapacity = NodeInfo.SendCapacity();
+	int64 MaxSend = SendCapacity - DrainLimit;
+
+	while (!WorkQueue.IsEmpty() && NodeInfo.CanSendSegments() && SendCapacity > DrainLimit )
+	{
+		if (TOptional<FSentSegmentInfoArray> ProcessedMessages = DrainQueue(NodeInfo, WorkQueue, MaxSend, bShouldDrainOverflowQueue))
+		{
+			bShouldDrainOverflowQueue = false;
+			BytesSent += RequeueWork(NodeInfo, WorkQueue, *ProcessedMessages);
+			SendCapacity = NodeInfo.SendCapacity();
+			MaxSend = SendCapacity - DrainLimit;
+		}
+		else
+		{
+			// Socket Error
+			return false;
+		}
+	}
+	return true;
+}
+
+int32 FUdpMessageProcessor::UpdateSegmenters(FNodeInfo& NodeInfo)
+{
+	SCOPED_MESSAGING_TRACE(FUdpMessageProcessor_UpdateSegmenters);
+
+	const float ReliableQueuePriority = GetDefault<UUdpMessagingSettings>()->ReliableQueuePriority / 100.f;
+
+	int32 BytesSent = 0;
+	int64 SendCapacity = NodeInfo.SendCapacity();
+	const int64 SendRoomForUnreliables = (1 - ReliableQueuePriority) * SendCapacity;
+	int64 MaxSendForReliables = SendCapacity - SendRoomForUnreliables;
+
+	// Prioritize the reliable queue first. The reliable queue always gets more bandwidth.  We drain the overflow queue first and only
+	// on the first run as we don't want it to monopolize our time either.
+	// 
+	if (!ProcessQueue(NodeInfo, NodeInfo.ReliableWorkQueue, SendRoomForUnreliables, BytesSent, true))
+	{
+		// Socket error 
+		return -1;
+	}
 
-	if (CVarCheckForExpiredWithFullQueue.GetValueOnAnyThread() > 0 && !NodeInfo.CanSendSegments() && NodeInfo.WorkQueue.IsFull())
+	if (!ProcessQueue(NodeInfo, NodeInfo.UnreliableWorkQueue, 0, BytesSent))
+	{
+		// Socket error 
+		return -1;
+	}
+
+	if (NodeInfo.SendCapacity() > 0)
+	{
+		// If there is still more room left go back to work on the reliable queue until we have nothing left or run out
+		// of capacity.
+		if (!ProcessQueue(NodeInfo, NodeInfo.ReliableWorkQueue, 0, BytesSent))
+		{
+			// Socket error 
+			return -1;
+		}
+	}
+	
+	if (CVarCheckForExpiredWithFullQueue.GetValueOnAnyThread() > 0 && !NodeInfo.CanSendSegments() && NodeInfo.ReliableWorkQueue.IsFull())
 	{
 		UE_LOG(LogUdpMessaging, Warning, TEXT("Work queue is full sending to node at address %s. We cannot send new data. Attempting to expire old inflight segments."), *NodeInfo.Endpoint.ToString());
 		// RemoveLostSegments will clean-up the tracking map that we use to determine what is "inflight". Old data will be expired and scheduled
@@ -1354,8 +1454,6 @@ int32 FUdpMessageProcessor::UpdateSegmenters(FNodeInfo& NodeInfo)
 	return BytesSent;
 }
 
-
-
 const FTimespan FUdpMessageProcessor::StaleReassemblyInterval = FTimespan::FromSeconds(30);
 
 bool FUdpMessageProcessor::UpdateReassemblers(FNodeInfo& NodeInfo)
@@ -1421,7 +1519,7 @@ bool FUdpMessageProcessor::UpdateReassemblers(FNodeInfo& NodeInfo)
 			if (!ReassembledMessage->IsDelivered())
 			{
 				const int ReceivedSegments = ReassembledMessage->GetTotalSegmentsCount() - ReassembledMessage->GetPendingSegmentsCount();
-				UE_LOG(LogUdpMessaging, Warning, TEXT("FUdpMessageProcessor::UpdateReassemblers Discarding %d/%d of stale message segements from %s"),
+				UE_LOG(LogUdpMessaging, Warning, TEXT("FUdpMessageProcessor::UpdateReassemblers Discarding %d/%d of stale message segments from %s"),
 					ReceivedSegments,
 					ReassembledMessage->GetTotalSegmentsCount(),
 					*ReassembledMessage->Describe());
diff --git a/Engine/Plugins/Messaging/UdpMessaging/Source/UdpMessaging/Private/Transport/UdpMessageProcessor.h b/Engine/Plugins/Messaging/UdpMessaging/Source/UdpMessaging/Private/Transport/UdpMessageProcessor.h
--- a/Engine/Plugins/Messaging/UdpMessaging/Source/UdpMessaging/Private/Transport/UdpMessageProcessor.h
+++ b/Engine/Plugins/Messaging/UdpMessaging/Source/UdpMessaging/Private/Transport/UdpMessageProcessor.h
@@ -4,6 +4,7 @@
 
 #include "Containers/MpscQueue.h"
 
+#include "HAL/Platform.h"
 #include "ModuleDescriptor.h"
 
 #include "HAL/Runnable.h"
@@ -35,6 +36,7 @@ class FUdpSerializedMessage;
 class FUdpSocketSender;
 class IMessageAttachment;
 enum class EUdpMessageFormat : uint8;
+enum class EUdpMessageSegments : uint8;
 using FUdpReassembledMessageSharedPtr = TSharedPtr<FUdpReassembledMessage, ESPMode::ThreadSafe>;
 
 /**
@@ -94,6 +96,7 @@ struct FSentSegmentInfo
 	bool bRequiresRequeue = false;
 	bool bSendSocketError = false;
 	bool bFullySent = false;
+	bool bIsTraced = false;
 
 	/** Converts FSentSegmentInfo into a FSentData struct. */
 	FSentData AsSentData(int32 InMessageId, uint32 SegmentNumber, const FDateTime& CurrentTime)
@@ -113,6 +116,8 @@ struct FSentSegmentInfo
 	}
 };
 
+using FSentSegmentInfoArray = TArray<FSentSegmentInfo, TInlineAllocator<64>>;
+
 namespace UE::Private::MessageProcessor
 {
 /**
@@ -155,8 +160,11 @@ class FUdpMessageProcessor
 		/** Holds the collection of message segmenters. */
 		TMap<int32, TSharedPtr<FUdpMessageSegmenter>> Segmenters;
 
-		/** Holds of queue of MessageIds to send. They are processed in round-robin fashion. */
-		TUdpCircularQueue<int32> WorkQueue;
+		/** Holds of queue of MessageIds to send for reliable messages. They are processed in round-robin fashion. */
+		TUdpCircularQueue<int32> ReliableWorkQueue;
+
+		/** Holds of queue of MessageIds to send for unreliable messages. They are processed in round-robin fashion. */
+		TUdpCircularQueue<int32> UnreliableWorkQueue;
 
 		/* Holds the list of messages that are still pending acknowledgement or overflow if overcommitted on work queue. */
 		TArray<int32> OverflowForPendingAck;
@@ -176,8 +184,8 @@ class FUdpMessageProcessor
 		/** Average time from when the packet was sent to ack received. */
 		FTimespan AvgRoundTripTime = FTimespan::FromMilliseconds(20);
 
-		/** Maximum number of segments we can have in flight. Start with a small value and let it grow based on performance. */
-		uint64 WindowSize = 64;
+		/** Maximum number of segments we can have in flight. Start with a middle of the road value and adjust based on network performance. */
+		uint64 WindowSize = 1024;
 
 		/** Various transport statistics for this endpoint */
 		FMessageTransportStatistics Statistics;
@@ -190,7 +198,7 @@ class FUdpMessageProcessor
 		 */
 		bool CanCommitToWorkQueue(const FDateTime& InCurrentTime)
 		{
-			const bool bWorkQueueFull = WorkQueue.IsFull();
+			const bool bWorkQueueFull = ReliableWorkQueue.IsFull();
 			if (bWorkQueueFull)
 			{
 				const double ReportingIntervalInSeconds = 1;
@@ -427,6 +435,11 @@ class FUdpMessageProcessor
 			return InflightSegments.Num() < WindowSize;
 		}
 
+		int64 SendCapacity() const
+		{
+			return static_cast<int64>(WindowSize) - static_cast<int64>(InflightSegments.Num());
+		}
+		
 		/** Resets the endpoint info. */
 		void ResetIfRestarted(const FGuid& NewNodeId)
 		{
@@ -631,6 +644,22 @@ public:
 		return ErrorDelegate;
 	}
 
+	void SetEnableMessageDelegates(bool bShouldEnable)
+	{
+		bEnableMessageDelegates = bShouldEnable;
+	}
+
+	/**
+	 * Inbound segment delegate inspector
+	 *
+	 * @return The delegate.
+	 */
+	DECLARE_TS_MULTICAST_DELEGATE_FourParams(FOnInboundSegment, const FInboundSegment&, const FGuid&, EUdpMessageSegments, bool)
+	FOnInboundSegment& OnInboundSegmentReceived_UdpMessageProcessorThread()
+	{
+		return InboundSegmentDelegate;
+	}
+	
 	/** 
 	 * Returns a delegate that is executed when a socket fails to communicate
 	 * upon sending to a target endpoint.
@@ -841,14 +870,41 @@ protected:
 	/** Updates nodes per protocol version map */
 	void UpdateNodesPerVersion();
 
+
 protected:
 
 	//~ FSingleThreadRunnable interface
 
 	virtual void Tick() override;
 
+	/** Return a copy of the known nodes as seen by the processor. This is used by debugging logging utils and should not be relied on as an API*/
+	TMap<FGuid, FNodeInfo> GetKnownNodesAsSnapshot()
+	{
+		return KnownNodes;
+	}
+
 private:
+	friend struct FUdpMessageProcessorLogger;
+	friend struct FUdpMessageLoggingImpl;
 
+	/**
+	 * Drain a given send queue to the drain limit specified on the parameter listed. Also drain the overflow bucket if requested. Overflow
+	 * may occasionally fill from pending acks or work queue overflowing.  
+	 */
+	TOptional<FSentSegmentInfoArray> DrainQueue(FNodeInfo& NodeInfo, TUdpCircularQueue<int32>& WorkQueue, int64 DrainCount, const bool bDrainOverflow);
+
+	/**
+	 * Refill the given queue with new work based on the processed messages. They will be put back in the queue to be processed later
+	 * in the next round to ensure fairness in the sending of data.
+	 */
+	int32 RequeueWork(FNodeInfo& NodeInfo, TUdpCircularQueue<int32>& WorkQueue, const FSentSegmentInfoArray& ProcessedMessages);
+
+	/**
+	 * Process a given queue.  This calls into DrainQueue and RequeueWork repeatedly until we reach the drain limit which is the minimum
+	 * value of of send capacity.  For example, if DrainLimit is 0 then we work until the send capacity is 0. 
+	 */
+	bool ProcessQueue(FNodeInfo& NodeInfo, TUdpCircularQueue<int32>& WorkQueue, int64 DrainLimit, int32& BytesSent, bool bShouldDrainOverflowQueue = false);
+		
 	/** Send our list of known nodes to the known nodes. */
 	void SendKnownNodesToKnownNodes();
 
@@ -954,6 +1010,12 @@ private:
 	/** Holds a delegate to be invoked when checking the validity of a given endpoint address. */
 	FCanAcceptEndpoint CanAcceptEndpointDelegate;
 
+	/** Holds a delegate to allow logger to inspect inbound segments. */
+	FOnInboundSegment InboundSegmentDelegate;
+
+	/** Enable / disable inbound segment delegates. */
+	std::atomic<bool> bEnableMessageDelegates = false;
+
 	/** The configured message format (from UUdpMessagingSettings). */
 	EUdpMessageFormat MessageFormat;
 
diff --git a/Engine/Plugins/Messaging/UdpMessaging/Source/UdpMessaging/Private/Transport/UdpMessageProcessorLogger.cpp b/Engine/Plugins/Messaging/UdpMessaging/Source/UdpMessaging/Private/Transport/UdpMessageProcessorLogger.cpp
new file mode 100644
--- /dev/null
+++ b/Engine/Plugins/Messaging/UdpMessaging/Source/UdpMessaging/Private/Transport/UdpMessageProcessorLogger.cpp
@@ -0,0 +1,623 @@
+// Copyright Epic Games, Inc. All Rights Reserved.
+
+// Implements a very simple CSV style logger that is supplimental to the logging that the engine uses.  When enabled, it will record all inbound, outbound segment data.  This
+// is controlled via a CVar.  There is also a variable that will only log trace packets (UdpProcessor.LogOnlyTracedPackets).
+// 
+#include "UdpMessageProcessorLogger.h"
+
+#include "Containers/Ticker.h"
+#include "Containers/SpscQueue.h"
+#include "Interfaces/IPv4/IPv4Endpoint.h"
+
+#include "Misc/DateTime.h"
+#include "Misc/Guid.h"
+#include "Misc/ScopeExit.h"
+#include "Misc/Paths.h"
+#include "Misc/StringBuilder.h"
+#include "Misc/App.h"
+
+#include "Serialization/Archive.h"
+#include "Serialization/ArrayReader.h"
+
+#include "HAL/FileManager.h"
+#include "HAL/IConsoleManager.h"
+
+#include "Features/IModularFeatures.h"
+
+#include "INetworkMessagingExtension.h"
+
+#include "Templates/UniquePtr.h"
+
+#include "Shared/UdpMessageSegment.h"
+
+#include "UdpMessageProcessor.h"
+#include "UdpMessagingPrivate.h"
+
+struct FNodeInfoStats
+{
+	/** Holds the node's IP endpoint. */
+	FIPv4Endpoint Endpoint;
+
+	/** Holds the time at which the last Hello was received. */
+	FDateTime LastSegmentReceivedTime;
+
+	/** Holds the endpoint's node identifier. */
+	FGuid NodeId;
+
+	/** Holds the protocol version this node is communicating with */
+	uint8 ProtocolVersion;
+
+	/** This object is always traced. */
+	bool bTraced = true;
+
+	/** This stat is always marked complete so it gets processed by the logger. */
+	bool bComplete = true;
+};
+
+struct FInboundSegmentStat
+{
+	/** Sender endpoint */
+	FIPv4Endpoint InEndpoint;
+
+	/** Guid for the node. */
+	FGuid NodeId;
+
+	/** Segment type */
+	EUdpMessageSegments SegmentType;
+
+	/** Was this packet filtered.  */
+	bool bFiltered = false;
+	
+	/** Size of the payload. */
+	int64 Size;
+
+	/** By default we never trace inbound segment stats. */
+	bool bTraced = false;
+	
+	/** Inbound segments are always complete for the purposes of logging below. */
+	bool bComplete = true;
+};
+
+namespace UE::UdpMessaging::Private
+{
+static int32 bLogOnlyTracedObjects = 0;
+static FAutoConsoleVariableRef  CVarLogOnlyTraced(TEXT("UdpProcessor.LogOnlyTracedPackets"), bLogOnlyTracedObjects, TEXT("Log only segments and reassemblers that have a tracing enabled."));
+	
+static INetworkMessagingExtension* GetMessagingStatistics()
+{
+	IModularFeatures& ModularFeatures = IModularFeatures::Get();
+	if (IsInGameThread())
+	{
+		if (ModularFeatures.IsModularFeatureAvailable(INetworkMessagingExtension::ModularFeatureName))
+		{
+			return &ModularFeatures.GetModularFeature<INetworkMessagingExtension>(INetworkMessagingExtension::ModularFeatureName);
+		}
+	}
+	else
+	{
+		ModularFeatures.LockModularFeatureList();
+		ON_SCOPE_EXIT
+		{
+			ModularFeatures.UnlockModularFeatureList();
+		};
+		
+		if (ModularFeatures.IsModularFeatureAvailable(INetworkMessagingExtension::ModularFeatureName))
+		{
+			return &ModularFeatures.GetModularFeature<INetworkMessagingExtension>(INetworkMessagingExtension::ModularFeatureName);
+		}
+	}
+		
+	ensureMsgf(false, TEXT("Feature %s is unavailable"), *INetworkMessagingExtension::ModularFeatureName.ToString());
+	return nullptr;
+}
+
+template<typename TStatLogger>
+FDelegateHandle Register(TStatLogger*);
+
+template<typename TStatLogger>
+void Unregister(TStatLogger*);
+
+template<typename TStatLogger>
+FString GetOutputName(TStatLogger*);
+
+template<typename TStatLogger>
+FString GetHeader(TStatLogger*);
+
+template<typename TStatType>
+FString LogOne(const TStatType&);
+	
+}; // namespace UE::UdpMessaging::Private
+
+template <typename TStatType>
+struct FUdpMessageProcessorLoggerTyped
+{
+	FUdpMessageProcessorLoggerTyped(const FIPv4Endpoint& Endpoint)
+	{
+		UE::UdpMessaging::Private::Register(this);
+		TickHandler = FTSTicker::GetCoreTicker().AddTicker(FTickerDelegate::CreateRaw(this, &FUdpMessageProcessorLoggerTyped<TStatType>::GameTick));
+		StartLogging(Endpoint);
+	}
+
+	~FUdpMessageProcessorLoggerTyped()
+	{
+		UE::UdpMessaging::Private::Unregister(this);
+		FTSTicker::GetCoreTicker().RemoveTicker(TickHandler);
+		GameTick(0);
+
+		Flush();
+	}
+
+	void Flush()
+	{
+		if (CSVArchive.IsValid())
+		{
+			CSVArchive->Flush();
+		}
+	}
+
+	void StartLogging(const FIPv4Endpoint& Endpoint)
+	{
+		FString InstanceName = FApp::GetInstanceName();
+		FString EndpointStr = Endpoint.ToString().Replace(TEXT("."), TEXT("_")).Replace(TEXT(":"), TEXT("p"));
+		const FString Filename = FString::Printf(TEXT("%s-%s-%s-%s.csv"),*InstanceName, *EndpointStr, *UE::UdpMessaging::Private::GetOutputName(this), *FDateTime::Now().ToString());
+		const FString FullPath = FPaths::ProjectLogDir() / TEXT("UdpMessaging") / Filename;
+		CSVArchive = TUniquePtr<FArchive>(IFileManager::Get().CreateFileWriter(*FullPath, EFileWrite::FILEWRITE_AllowRead));
+		if (CSVArchive)
+		{
+			UTF8CHAR UTF8BOM[] = { (UTF8CHAR)0xEF, (UTF8CHAR)0xBB, (UTF8CHAR)0xBF };
+			CSVArchive->Serialize(&UTF8BOM, UE_ARRAY_COUNT(UTF8BOM) * sizeof(UTF8CHAR));
+		}
+
+		FString CSVHeader = TEXT("Time,") + UE::UdpMessaging::Private::GetHeader(this);
+		CSVHeader += LINE_TERMINATOR;
+		if (CSVArchive.IsValid())
+		{
+			FTCHARToUTF8 CSVHeaderUTF8(*CSVHeader);
+			CSVArchive->Serialize((UTF8CHAR*)CSVHeaderUTF8.Get(), CSVHeaderUTF8.Length() * sizeof(UTF8CHAR));
+		}
+	}
+
+	void StatUpdated(TStatType Stats)
+	{
+		FWrappedType WrappedType{FDateTime::Now(), MoveTemp(Stats)};
+		AsyncStatQueue.Enqueue(MoveTemp(WrappedType));
+	};
+
+	struct FWrappedType
+	{
+		FDateTime EventTime;
+		TStatType StatData;
+	};
+
+	bool GameTick(float DeltaTime)
+	{
+		if (!CSVArchive.IsValid())
+		{
+			return true;
+		}
+
+		FWrappedType Wrapped;
+		while (AsyncStatQueue.Dequeue(Wrapped))
+		{
+			if (!Wrapped.StatData.bComplete && !Wrapped.StatData.bTraced)
+			{
+				continue;
+			}
+			if (UE::UdpMessaging::Private::bLogOnlyTracedObjects > 0 && !Wrapped.StatData.bTraced)
+			{
+				continue;
+			}
+
+			FString CSVRow = Wrapped.EventTime.ToString() + TEXT(",") + UE::UdpMessaging::Private::LogOne(Wrapped.StatData);
+			CSVRow += LINE_TERMINATOR;
+			{
+				FTCHARToUTF8 CSVRowUTF8(*CSVRow);
+				CSVArchive->Serialize((UTF8CHAR*)CSVRowUTF8.Get(), CSVRowUTF8.Length() * sizeof(UTF8CHAR));
+			}
+		}
+		return true;
+	}
+
+	FTSTicker::FDelegateHandle TickHandler;
+	TSpscQueue<FWrappedType> AsyncStatQueue;
+	TUniquePtr<FArchive> CSVArchive;
+};
+
+
+namespace UE::UdpMessaging::Private
+{
+	template<typename TStatLogger>
+	FDelegateHandle Register(TStatLogger*)
+	{
+		return {};
+	}
+
+	template<typename TStatLogger>
+	void Unregister(TStatLogger*)
+	{
+	}
+
+	// Template specialization for outbound statistics
+	template <>
+	FDelegateHandle Register(FUdpMessageProcessorLoggerTyped<FOutboundTransferStatistics>* Logger)
+	{
+		if (INetworkMessagingExtension* Statistics = GetMessagingStatistics())
+		{
+			return Statistics->OnOutboundTransferUpdatedFromThread().AddRaw(Logger, &FUdpMessageProcessorLoggerTyped<FOutboundTransferStatistics>::StatUpdated);
+		}
+		return {};
+	}
+	// Template specialization for inbound statistics
+	template <>
+	FDelegateHandle Register(FUdpMessageProcessorLoggerTyped<FInboundTransferStatistics>* Logger)
+	{
+		if (INetworkMessagingExtension* Statistics = GetMessagingStatistics())
+		{
+			return Statistics->OnInboundTransferUpdatedFromThread().AddRaw(Logger, &FUdpMessageProcessorLoggerTyped<FInboundTransferStatistics>::StatUpdated);
+		}
+		return {};
+	}
+	// Template specialization for outbound statistics
+	template <>
+	void Unregister(FUdpMessageProcessorLoggerTyped<FOutboundTransferStatistics>* Logger)
+	{
+		if (INetworkMessagingExtension* Statistics = GetMessagingStatistics())
+		{
+			Statistics->OnOutboundTransferUpdatedFromThread().RemoveAll(Logger);
+		}
+	}
+
+	// Template specialization for inbound statistics
+	template <>
+	void Unregister(FUdpMessageProcessorLoggerTyped<FInboundTransferStatistics>* Logger)
+	{
+		if (INetworkMessagingExtension* Statistics = GetMessagingStatistics())
+		{
+			Statistics->OnInboundTransferUpdatedFromThread().RemoveAll(Logger);
+		}
+	}
+
+	template <>
+	FString GetOutputName(FUdpMessageProcessorLoggerTyped<FInboundSegmentStat>* Logger)
+	{
+		return TEXT("InboundSegmentStat");
+	}
+
+	template <>
+	FString GetOutputName(FUdpMessageProcessorLoggerTyped<FNodeInfoStats>* Logger)
+	{
+		return TEXT("NodeInfo");
+	}
+
+	template <>
+	FString GetOutputName(FUdpMessageProcessorLoggerTyped<FOutboundTransferStatistics>* Logger)
+	{
+		return TEXT("Outbound");
+	}
+	
+	template<>
+	FString GetOutputName(FUdpMessageProcessorLoggerTyped<FInboundTransferStatistics>* Logger)
+	{
+		return TEXT("Inbound");
+	}
+
+	template <>
+	FString GetHeader(FUdpMessageProcessorLoggerTyped<FInboundSegmentStat>* Logger)
+	{
+		TStringBuilder<512> Builder;
+
+		Builder.Append(TEXT("InEndpoint,"));
+		Builder.Append(TEXT("NodeId,"));
+		Builder.Append(TEXT("SegmentType,"));
+		Builder.Append(TEXT("Filtered,"));
+		Builder.Append(TEXT("Size"));
+		return Builder.ToString();
+	}
+
+	template <>
+	FString GetHeader(FUdpMessageProcessorLoggerTyped<FOutboundTransferStatistics>* Logger)
+	{
+		TStringBuilder<512> Builder;
+		Builder.Append(TEXT("DestinationId,"));
+		Builder.Append(TEXT("MessageId,"));
+		Builder.Append(TEXT("BytesToSend,"));
+		Builder.Append(TEXT("BytesSent,"));
+		Builder.Append(TEXT("BytesAcknowledged,"));
+		Builder.Append(TEXT("bIsReliable,"));
+		Builder.Append(TEXT("bComplete"));
+		return Builder.ToString();
+	}
+
+	template<>
+	FString GetHeader(FUdpMessageProcessorLoggerTyped<FInboundTransferStatistics>* Logger)
+	{
+		TStringBuilder<512> Builder;
+		Builder.Append(TEXT("OriginId,"));
+		Builder.Append(TEXT("MessageId,"));
+		Builder.Append(TEXT("BytesToReceive,"));
+		Builder.Append(TEXT("BytesReceived,"));
+		Builder.Append(TEXT("bWasReliablySent,"));
+		Builder.Append(TEXT("bComplete"));
+		return Builder.ToString();
+	}
+
+	template<>
+	FString GetHeader(FUdpMessageProcessorLoggerTyped<FNodeInfoStats>* Logger)
+	{
+		TStringBuilder<512> Builder;
+		Builder.Append(TEXT("Endpoint,"));
+		Builder.Append(TEXT("LastSegmentReceivedTime,"));
+		Builder.Append(TEXT("NodeId,"));
+		Builder.Append(TEXT("ProtocolVersion"));
+		return Builder.ToString();
+	}
+
+	template<>
+	FString LogOne(const FInboundSegmentStat& Stat)
+	{
+		TStringBuilder<512> Builder;
+		Builder.Append(Stat.InEndpoint.ToString());
+		Builder.Append(TEXT(","));
+		Builder.Append(Stat.NodeId.ToString());
+		Builder.Append(TEXT(","));
+
+		switch (Stat.SegmentType)
+		{
+			case EUdpMessageSegments::Abort:
+				Builder.Append(TEXT("Abort"));
+				break;
+
+			case EUdpMessageSegments::Acknowledge:
+				Builder.Append(TEXT("Ack"));
+				break;
+
+			case EUdpMessageSegments::AcknowledgeSegments:
+				Builder.Append(TEXT("AckSegments"));
+				break;
+
+			case EUdpMessageSegments::Bye:
+				Builder.Append( TEXT("Bye"));
+				break;
+
+			case EUdpMessageSegments::Data:
+				Builder.Append( TEXT("Data") );
+				break;
+
+			case EUdpMessageSegments::Hello:
+				Builder.Append(TEXT("Hello"));
+				break;
+
+			case EUdpMessageSegments::Ping:
+				Builder.Append(TEXT("Ping"));
+				break;
+
+			case EUdpMessageSegments::Pong:
+				Builder.Append(TEXT("Pong"));
+				break;
+
+			case EUdpMessageSegments::Retransmit:
+				Builder.Append(TEXT("Retransmit"));
+				break;
+
+			case EUdpMessageSegments::Timeout:
+				Builder.Append(TEXT("Timeout"));
+				break;
+
+			case EUdpMessageSegments::Mesh:
+				Builder.Append(TEXT("Mesh"));
+				break;
+
+			default:
+				Builder.Append(TEXT("Unknown"));								
+		};
+		
+		Builder.Append(TEXT(","));
+		Builder.Append(FString::FromInt(Stat.bFiltered));
+		Builder.Append(TEXT(","));
+		Builder.Append(FString::FromInt(Stat.Size));
+		return Builder.ToString();
+		
+	}
+
+	template<>
+	FString LogOne(const FNodeInfoStats& Stat)
+	{
+		TStringBuilder<512> Builder;
+		Builder.Append(Stat.Endpoint.ToString());
+		Builder.Append(TEXT(","));
+		Builder.Append(Stat.LastSegmentReceivedTime.ToString());
+		Builder.Append(TEXT(","));
+		Builder.Append(Stat.NodeId.ToString());
+		Builder.Append(TEXT(","));
+		Builder.Append(FString::FromInt(Stat.ProtocolVersion));
+		return Builder.ToString();
+	}
+
+	template <>
+	FString LogOne(const FOutboundTransferStatistics& Stat)
+	{
+		TStringBuilder<512> Builder;
+		Builder.Append(Stat.DestinationId.ToString());
+		Builder.Append(TEXT(","));
+		Builder.Append(FString::FromInt(Stat.MessageId));
+		Builder.Append(TEXT(","));
+		Builder.Append(FString::FromInt(Stat.BytesToSend));
+		Builder.Append(TEXT(","));
+		Builder.Append(FString::FromInt(Stat.BytesSent));
+		Builder.Append(TEXT(","));
+		Builder.Append(FString::FromInt(Stat.BytesAcknowledged));
+		Builder.Append(TEXT(","));
+		Builder.Append(FString::FromInt(Stat.bIsReliable));
+		Builder.Append(TEXT(","));
+		Builder.Append(FString::FromInt(Stat.bComplete));
+		return Builder.ToString();
+	}
+	
+	template<>
+	FString LogOne(const FInboundTransferStatistics& Stat)
+	{
+		TStringBuilder<512> Builder;
+		Builder.Append(Stat.OriginId.ToString());
+		Builder.Append(TEXT(","));
+		Builder.Append(FString::FromInt(Stat.MessageId));
+		Builder.Append(TEXT(","));
+		Builder.Append(FString::FromInt(Stat.BytesToReceive));
+		Builder.Append(TEXT(","));
+		Builder.Append(FString::FromInt(Stat.BytesReceived));
+		Builder.Append(TEXT(","));
+		Builder.Append(FString::FromInt(Stat.bWasReliablySent));
+		Builder.Append(TEXT(","));
+		Builder.Append(FString::FromInt(Stat.bComplete));
+		return Builder.ToString();
+	}
+
+	TMap<FIPv4Endpoint, FUdpMessageProcessorLogger*> Loggers;
+	
+	void EnableLogging()
+	{
+		for (TMap<FIPv4Endpoint, FUdpMessageProcessorLogger*>::TIterator It(Loggers); It; ++It)
+		{
+			UE_LOG(LogUdpMessaging, Display, TEXT("Starting logging for %s."), *(It->Key.ToString()));
+			It->Value->StartLogging();
+		}
+	}		
+
+	void DisableLogging()
+	{
+		for (TMap<FIPv4Endpoint, FUdpMessageProcessorLogger*>::TIterator It(Loggers); It; ++It)
+		{
+			UE_LOG(LogUdpMessaging, Display, TEXT("Stopping logging for %s."), *(It->Key.ToString()));
+			It->Value->StopLogging();
+		}		
+	}
+
+	void KnownEndpoints()
+	{
+		for (TMap<FIPv4Endpoint, FUdpMessageProcessorLogger*>::TIterator It(Loggers); It; ++It)
+		{
+			UE_LOG(LogUdpMessaging, Display, TEXT("Dumping known nodes for %s."), *(It->Key.ToString()));
+			It->Value->DumpKnownNodeInfo();
+		}
+	}
+
+	FAutoConsoleCommand EnableLoggingCommand(TEXT("UdpProcessor.EnableLogging"), TEXT("Enable Additional Logging for UdpProcessor."), FConsoleCommandDelegate::CreateStatic(&EnableLogging));
+	FAutoConsoleCommand DisableLoggingCommand(TEXT("UdpProcessor.DisableLogging"), TEXT("Disable Additional Logging for UdpProcessor."), FConsoleCommandDelegate::CreateStatic(&DisableLogging));
+
+	FAutoConsoleCommand DumpKnownEndpoints(TEXT("UdpProcessor.DumpKnownEndpoints"), TEXT("Disable Additional Logging for UdpProcessor."), FConsoleCommandDelegate::CreateStatic(&KnownEndpoints));
+	
+
+} // namespace UE::UdpMessaging::Private
+
+
+/** Pointer to logging implementation. */
+struct FUdpMessageLoggingImpl
+{
+	FUdpMessageLoggingImpl(const FIPv4Endpoint& Endpoint, FUdpMessageProcessor* InProcessor)
+		: Inbound(Endpoint)
+		, Outbound(Endpoint)
+		, NodeInfo(Endpoint)
+		, InboundSegmentStat(Endpoint)
+		, Processor(InProcessor)
+	{
+	}
+
+	~FUdpMessageLoggingImpl()
+	{
+		Processor->SetEnableMessageDelegates(false);
+		
+		Processor->OnInboundSegmentReceived_UdpMessageProcessorThread().RemoveAll(this);
+	}
+
+	void AttachToProcessor()
+	{
+		Processor->OnInboundSegmentReceived_UdpMessageProcessorThread().AddRaw(this, &FUdpMessageLoggingImpl::InboundMessage);
+
+		Processor->SetEnableMessageDelegates(true);
+	}
+	
+	void InboundMessage(const FUdpMessageProcessor::FInboundSegment& InMessage, const FGuid& NodeId, EUdpMessageSegments SegmentType, bool bFiltered)
+	{
+		FInboundSegmentStat Stat = {
+			InMessage.Sender,
+			NodeId,
+			SegmentType,
+			bFiltered,
+			InMessage.Data ? InMessage.Data->TotalSize() : 0
+		};
+		InboundSegmentStat.StatUpdated(MoveTemp(Stat));
+	}
+	
+	FUdpMessageProcessorLoggerTyped<FInboundTransferStatistics> Inbound;
+	FUdpMessageProcessorLoggerTyped<FOutboundTransferStatistics> Outbound;
+	FUdpMessageProcessorLoggerTyped<FNodeInfoStats> NodeInfo;
+	FUdpMessageProcessorLoggerTyped<FInboundSegmentStat> InboundSegmentStat;
+
+	FUdpMessageProcessor* Processor = nullptr;
+};
+
+FUdpMessageProcessorLogger::FUdpMessageProcessorLogger(FIPv4Endpoint InEndpoint, uint16 PortNo, FUdpMessageProcessor* InProcessor)
+	: Processor(InProcessor)
+{
+	Endpoint = MoveTemp(InEndpoint);
+	Endpoint.Port = PortNo;
+
+	UE::UdpMessaging::Private::Loggers.Add(Endpoint, this);
+}
+
+FUdpMessageProcessorLogger::~FUdpMessageProcessorLogger()
+{
+	UE::UdpMessaging::Private::Loggers.Remove(Endpoint);
+}
+
+void FUdpMessageProcessorLogger::StartLogging()
+{
+	if (!Impl)
+	{
+		Impl = MakeUnique<FUdpMessageLoggingImpl>(Endpoint, Processor);
+
+		DumpKnownNodeInfo();
+		
+		Impl->NodeInfo.Flush();
+		
+		// Attach to the inbound and out bound processors. 
+		Impl->AttachToProcessor();		
+	}
+}
+
+void FUdpMessageProcessorLogger::StopLogging()
+{
+	Impl.Reset();
+}
+
+bool FUdpMessageProcessorLogger::IsLogging() const
+{
+	return Impl.IsValid();
+}
+
+void FUdpMessageProcessorLogger::DumpKnownNodeInfo()
+{
+	auto SnapshotStats = [this](FUdpMessageProcessorLoggerTyped<FNodeInfoStats> & NodeInfo)
+	{
+		TMap<FGuid, FUdpMessageProcessor::FNodeInfo> Snapshot = Processor->GetKnownNodesAsSnapshot();
+		for (auto It = Snapshot.CreateIterator(); It; ++It)
+		{
+			FNodeInfoStats Stat;
+			Stat.NodeId = It->Value.NodeId;
+			Stat.Endpoint = It->Value.Endpoint;
+			Stat.LastSegmentReceivedTime = It->Value.LastSegmentReceivedTime;
+			Stat.ProtocolVersion = It->Value.ProtocolVersion;
+
+			NodeInfo.StatUpdated(MoveTemp(Stat));
+		}
+	};
+	
+	if (Impl.IsValid())
+	{
+		SnapshotStats(Impl->NodeInfo);
+	}
+	else
+	{
+		FUdpMessageProcessorLoggerTyped<FNodeInfoStats> NodeInfo(Endpoint);
+		SnapshotStats(NodeInfo);
+	}
+}
diff --git a/Engine/Plugins/Messaging/UdpMessaging/Source/UdpMessaging/Private/Transport/UdpMessageProcessorLogger.h b/Engine/Plugins/Messaging/UdpMessaging/Source/UdpMessaging/Private/Transport/UdpMessageProcessorLogger.h
new file mode 100644
--- /dev/null
+++ b/Engine/Plugins/Messaging/UdpMessaging/Source/UdpMessaging/Private/Transport/UdpMessageProcessorLogger.h
@@ -0,0 +1,35 @@
+// Copyright Epic Games, Inc. All Rights Reserved.
+
+#pragma once
+
+#include "Interfaces/IPv4/IPv4Endpoint.h"
+#include "Templates/UniquePtr.h"
+#include "UdpMessageProcessor.h"
+
+struct FUdpMessageLoggingImpl;
+
+struct FUdpMessageProcessorLogger
+{
+	FUdpMessageProcessorLogger(FIPv4Endpoint InEndpoint, uint16 PortNo, FUdpMessageProcessor* InProcessor);
+	~FUdpMessageProcessorLogger();
+
+	/** Start enhanced logging for UdpMessageProcessor */
+	void StartLogging();
+
+	/** Stop enhanced logging for UdpMessageProcessor */
+	void StopLogging();
+
+	/** Report the current logging state. */
+	bool IsLogging() const;
+
+	/** Dump a report of the current known nodes from the system. */
+	void DumpKnownNodeInfo();
+	
+private:
+	FIPv4Endpoint Endpoint;
+		
+	FUdpMessageProcessor* Processor;
+	TUniquePtr<FUdpMessageLoggingImpl> Impl;
+};
+
+
diff --git a/Engine/Plugins/Messaging/UdpMessaging/Source/UdpMessaging/Private/Transport/UdpMessageTransport.cpp b/Engine/Plugins/Messaging/UdpMessaging/Source/UdpMessaging/Private/Transport/UdpMessageTransport.cpp
--- a/Engine/Plugins/Messaging/UdpMessaging/Source/UdpMessaging/Private/Transport/UdpMessageTransport.cpp
+++ b/Engine/Plugins/Messaging/UdpMessaging/Source/UdpMessaging/Private/Transport/UdpMessageTransport.cpp
@@ -17,6 +17,10 @@
 #include "Transport/UdpMessageProcessor.h"
 #include "UdpMessageProcessor.h"
 
+#if PLATFORM_DESKTOP
+#include "UdpMessageProcessorLogger.h"
+#endif
+
 /* FUdpMessageTransport structors
  *****************************************************************************/
 
@@ -133,6 +137,12 @@ void FUdpMessageTransport::OnAppPreExit()
 {
 	if (MessageProcessor)
 	{
+#if PLATFORM_DESKTOP
+		if (MessageLogger)
+		{
+			MessageLogger->StopLogging();
+		}
+#endif
 		MessageProcessor->WaitAsyncTaskCompletion();
 	}
 }
@@ -372,6 +382,7 @@ bool FUdpMessageTransport::StartTransport(IMessageTransportHandler& Handler)
 		UnicastReceiver->SetMaxReadBufferSize(2048);
 		UnicastReceiver->Start();
 	}
+	MessageLogger = MakeUnique<FUdpMessageProcessorLogger>(UnicastEndpoint, UnicastSocket ? UnicastSocket->GetPortNo() : 0, MessageProcessor);
 #endif
 
 	UE_LOG(LogUdpMessaging, Verbose, TEXT("Started Transport"));
diff --git a/Engine/Plugins/Messaging/UdpMessaging/Source/UdpMessaging/Private/Transport/UdpMessageTransport.h b/Engine/Plugins/Messaging/UdpMessaging/Source/UdpMessaging/Private/Transport/UdpMessageTransport.h
--- a/Engine/Plugins/Messaging/UdpMessaging/Source/UdpMessaging/Private/Transport/UdpMessageTransport.h
+++ b/Engine/Plugins/Messaging/UdpMessaging/Source/UdpMessaging/Private/Transport/UdpMessageTransport.h
@@ -23,6 +23,7 @@ class IMessageAttachment;
 class IMessageContext;
 class IMessageTransportHandler;
 class ISocketSubsystem;
+struct FUdpMessageProcessorLogger;
 
 struct FGuid;
 
@@ -191,4 +192,9 @@ private:
 
 	/** Async task queue for running deserialization of reassembled messages. */
 	TUniquePtr<UE::Tasks::FTaskConcurrencyLimiter> DeserializationTasks;
+
+#if PLATFORM_DESKTOP
+	/** Logging class to enable inspection of message data. */
+	TUniquePtr<FUdpMessageProcessorLogger> MessageLogger;
+#endif
 };
diff --git a/Engine/Plugins/Messaging/UdpMessaging/Source/UdpMessaging/Private/Transport/UdpReassembledMessage.h b/Engine/Plugins/Messaging/UdpMessaging/Source/UdpMessaging/Private/Transport/UdpReassembledMessage.h
--- a/Engine/Plugins/Messaging/UdpMessaging/Source/UdpMessaging/Private/Transport/UdpReassembledMessage.h
+++ b/Engine/Plugins/Messaging/UdpMessaging/Source/UdpMessaging/Private/Transport/UdpReassembledMessage.h
@@ -298,6 +298,12 @@ public:
 		}
 	}
 
+	/** True if this message is reliably delivered. */
+	const bool IsReliable() const
+	{
+		return EnumHasAnyFlags(GetFlags(), EMessageFlags::Reliable);
+	}
+
 	/**
 	 * Indicates whether or not this message was found to be malformed.
 	 * Once we detect a message is malformed, we will not attempt to process any data for it.
diff --git a/Engine/Plugins/Messaging/UdpMessaging/Source/UdpMessaging/Public/Shared/UdpMessagingSettings.h b/Engine/Plugins/Messaging/UdpMessaging/Source/UdpMessaging/Public/Shared/UdpMessagingSettings.h
--- a/Engine/Plugins/Messaging/UdpMessaging/Source/UdpMessaging/Public/Shared/UdpMessagingSettings.h
+++ b/Engine/Plugins/Messaging/UdpMessaging/Source/UdpMessaging/Public/Shared/UdpMessagingSettings.h
@@ -89,6 +89,13 @@ public:
 	UPROPERTY(config, EditAnywhere, Category = Transport, meta = (ClampMin="1024"), AdvancedDisplay)
 	uint16 WorkQueueSize = 1024;
 
+	/**
+	 * Reliable queue priority specified as a percentage with a range from 5-95%. Setting values above or below
+	 * can completely starve unreliable messages.
+	 */
+	UPROPERTY(config, EditAnywhere, Category = Transport, meta = (ClampMin="5", ClampMax="95"), AdvancedDisplay)
+	float ReliableQueuePriority = 75.0f;
+
 	/** Whether to stop the transport service when the application deactivates, and restart it when the application is reactivated */
 	UPROPERTY(config)
 	bool bStopServiceWhenAppDeactivates = true;
@@ -158,6 +165,14 @@ public:
 	UPROPERTY(config, EditAnywhere, Category=Transport, meta = (ClampMax="32"), AdvancedDisplay)
 	uint32 MaxConcurrentDeserializationTasks = 8;
 
+	/**
+	 * The period of time to consider an active connection to be timed-out. A regular heart beat is expected from connected endpoints and this will
+	 * expire that connection if no data is received after a period of time.  If zero is specified the endpoints are never timed out (this should
+	 * only be used for debugging purposes).
+	 */
+	UPROPERTY(config, EditAnywhere, Category=Transport, meta = (ClampMin="0"), AdvancedDisplay)
+	float ConnectionTimeoutPeriod = 5.0f;
+
 public:
 
 	/** Whether the UDP tunnel is enabled. */
diff --git a/Engine/Source/Runtime/Messaging/Public/IMessageContext.h b/Engine/Source/Runtime/Messaging/Public/IMessageContext.h
--- a/Engine/Source/Runtime/Messaging/Public/IMessageContext.h
+++ b/Engine/Source/Runtime/Messaging/Public/IMessageContext.h
@@ -186,8 +186,12 @@ enum class EMessageFlags : uint32
 {
 	/** No special flags */
 	None = 0,
+		
 	/** Guarantee that this message is delivered */
 	Reliable = 1 << 0,
+		
+	/** Trace this object through the transport on message bus. Only supported if the underlying transport handles the flag. */
+	Traced = 1 << 1
 };
 ENUM_CLASS_FLAGS(EMessageFlags);
 
diff --git a/Engine/Source/Runtime/Messaging/Public/INetworkMessagingExtension.h b/Engine/Source/Runtime/Messaging/Public/INetworkMessagingExtension.h
--- a/Engine/Source/Runtime/Messaging/Public/INetworkMessagingExtension.h
+++ b/Engine/Source/Runtime/Messaging/Public/INetworkMessagingExtension.h
@@ -67,6 +67,15 @@ struct FOutboundTransferStatistics
 
 	/** Total number of bytes acknowledge by the destination. */
 	uint64 BytesAcknowledged;
+
+	/** Indicate if the message is reliably sent. */
+	bool bIsReliable;
+
+	/** Indicates this was a traced packet. */
+	bool bTraced;
+	
+	/** Indicates if the message has been completely sent. In unreliable case, it means that all bytes have been sent. In the reliable case, it means it has been fully acknowledged. */
+	bool bComplete;
 };
 
 /** Per-node per-message transfer statistics. */
@@ -83,6 +92,15 @@ struct FInboundTransferStatistics
 
 	/** Total bytes sent thus far. */
 	int32 BytesReceived;
+
+	/** Was the message reliable. */
+	bool bWasReliablySent;
+
+	/** True, if this was a traced packet */
+	bool bTraced;
+	
+	/** True, if we have fully received the message */
+	bool bComplete;
 };
 
 
