Back out the changes to the beacon and use the old code path. The use of socket sender was forcing a failure when the failure was a valid result. This can occur when trying to send a Ping message as a response to the Hello on a subscribed multicast group which can happen if an editor is on a loopback device but sending multicast packets that the unicast on the receiving editor can receive.  This this case we cannot respond to those nodes so a ping failure is valid.

#rnx

diff --git a/Engine/Plugins/Messaging/UdpMessaging/Source/UdpMessaging/Private/Transport/UdpMessageBeacon.cpp b/Engine/Plugins/Messaging/UdpMessaging/Source/UdpMessaging/Private/Transport/UdpMessageBeacon.cpp
--- a/Engine/Plugins/Messaging/UdpMessaging/Source/UdpMessaging/Private/Transport/UdpMessageBeacon.cpp
+++ b/Engine/Plugins/Messaging/UdpMessaging/Source/UdpMessaging/Private/Transport/UdpMessageBeacon.cpp
@@ -22,17 +22,17 @@ const FTimespan FUdpMessageBeacon::MinimumInterval = FTimespan::FromMilliseconds
 /* FUdpMessageHelloSender structors
  *****************************************************************************/
 
-FUdpMessageBeacon::FUdpMessageBeacon(FUdpSocketSender* InSocketSender, const FGuid& InSocketId, const FIPv4Endpoint& InMulticastEndpoint)
+FUdpMessageBeacon::FUdpMessageBeacon(FSocket* InSocket, const FGuid& InSocketId, const FIPv4Endpoint& InMulticastEndpoint)
 	: BeaconInterval(MinimumInterval)
-	, LastEndpointCount(0)
+	, LastEndpointCount(1)
 	, LastHelloSent(FDateTime::MinValue())
 	, NextHelloTime(FDateTime::UtcNow())
 	, NodeId(InSocketId)
-	, SocketSender(InSocketSender)
+	, Socket(InSocket)
 	, Stopping(false)
 	, bSocketError(false)
 {
-	check(SocketSender);
+	check(Socket);
 	EndpointLeftEvent = FPlatformProcess::GetSynchEventFromPool(false);
 	MulticastAddress = InMulticastEndpoint.ToInternetAddr();
 	AddLocalEndpoints();
@@ -145,15 +145,23 @@ bool FUdpMessageBeacon::SendSegment(EUdpMessageSegments SegmentType, const FTime
 		Header.SegmentType = SegmentType;
 	}
 
-	TSharedRef<FArrayWriter, ESPMode::ThreadSafe> Writer = MakeShared<FArrayWriter, ESPMode::ThreadSafe>();
+	FArrayWriter Writer;
+	{
+		Writer << Header;
+		Writer << NodeId;
+	}
+
+	int32 Sent;
 
-	*Writer << Header;
-	*Writer << NodeId;
+	if (!Socket->Wait(ESocketWaitConditions::WaitForWrite, SocketWaitTime))
+	{
+		return false; // socket not ready for sending
+	}
 
-	if (!SocketSender->Send(Writer, MulticastAddress))
+	if (!Socket->SendTo(Writer.GetData(), Writer.Num(), Sent, *MulticastAddress))
 	{
 		bSocketError = true;
-		return false;
+		return false; // send failed
 	}
 
 	return true;
@@ -171,22 +179,30 @@ bool FUdpMessageBeacon::SendPing(const FTimespan& SocketWaitTime)
 	}
 	uint8 ActualProtocolVersion = UDP_MESSAGING_TRANSPORT_PROTOCOL_VERSION;
 
-	TSharedRef<FArrayWriter, ESPMode::ThreadSafe> Writer = MakeShared<FArrayWriter, ESPMode::ThreadSafe>();
+	
+	FArrayWriter Writer;
 	{
-		*Writer << Header;
-		*Writer << NodeId;
-		*Writer << ActualProtocolVersion; // Send our actual Protocol version as part of the ping message
+		Writer << Header;
+		Writer << NodeId;
+		Writer << ActualProtocolVersion; // Send our actual Protocol version as part of the ping message
 	}
 
+	if (!Socket->Wait(ESocketWaitConditions::WaitForWrite, SocketWaitTime))
+	{
+		return false; // socket not ready for sending
+	}
 
+	int32 Sent;
 	bool Result = true;
 	for (const auto& StaticAddress : StaticAddresses)
 	{
-		if (!SocketSender->Send(Writer, StaticAddress))
+		if (!Socket->SendTo(Writer.GetData(), Writer.Num(), Sent, *StaticAddress))
 		{	
 			Result = false; // send failed
 		}
+
 	}
+
 	return Result;
 }
 
diff --git a/Engine/Plugins/Messaging/UdpMessaging/Source/UdpMessaging/Private/Transport/UdpMessageBeacon.h b/Engine/Plugins/Messaging/UdpMessaging/Source/UdpMessaging/Private/Transport/UdpMessageBeacon.h
--- a/Engine/Plugins/Messaging/UdpMessaging/Source/UdpMessaging/Private/Transport/UdpMessageBeacon.h
+++ b/Engine/Plugins/Messaging/UdpMessaging/Source/UdpMessaging/Private/Transport/UdpMessageBeacon.h
@@ -29,12 +29,12 @@ public:
 	/** 
 	 * Creates and initializes a new beacon sender.
 	 *
-	 * @param InSocket The network socket sender used to send Hello segments.
+	 * @param InSocket The network socket used to send Hello segments.
 	 * @param InSocketId The network socket identifier (used to detect unicast endpoint).
 	 * @param InMulticastEndpoint The multicast group endpoint to transport messages to.
 	 * @param InStaticEndpoints The static nodes to broadcast to alongside the multicast.
 	 */
-	FUdpMessageBeacon(FUdpSocketSender* InSocketSender, const FGuid& InSocketId, const FIPv4Endpoint& InMulticastEndpoint);
+	FUdpMessageBeacon(FSocket* InSocketSender, const FGuid& InSocketId, const FIPv4Endpoint& InMulticastEndpoint);
 
 	/** Virtual destructor. */
 	virtual ~FUdpMessageBeacon();
@@ -158,8 +158,8 @@ private:
 	/** Holds local node identifier. */
 	FGuid NodeId;
 
-	/** Holds the socket sender used to send Hello segments. */
-	FUdpSocketSender* SocketSender;
+	/** Holds the socket used to send Hello segments. */
+	FSocket* Socket;
 
 	/** Holds a flag indicating that the thread is stopping. */
 	bool Stopping;
diff --git a/Engine/Plugins/Messaging/UdpMessaging/Source/UdpMessaging/Private/Transport/UdpMessageProcessor.cpp b/Engine/Plugins/Messaging/UdpMessaging/Source/UdpMessaging/Private/Transport/UdpMessageProcessor.cpp
--- a/Engine/Plugins/Messaging/UdpMessaging/Source/UdpMessaging/Private/Transport/UdpMessageProcessor.cpp
+++ b/Engine/Plugins/Messaging/UdpMessaging/Source/UdpMessaging/Private/Transport/UdpMessageProcessor.cpp
@@ -345,7 +345,7 @@ bool FUdpMessageProcessor::Init()
 	{
 		SocketSender = MakeUnique<FUdpSocketSender>(Socket, TEXT("FUdpMessageProcessor.Sender"));
 
-		Beacon = MakeUnique<FUdpMessageBeacon>(SocketSender.Get(), LocalNodeId, MulticastEndpoint);
+		Beacon = MakeUnique<FUdpMessageBeacon>(Socket, LocalNodeId, MulticastEndpoint);
 
 		// Current protocol version 17
 		SupportedProtocolVersions.Add(UDP_MESSAGING_TRANSPORT_PROTOCOL_VERSION);
