Ensure we don't call FindObject when saving or garbage collecting.

#rnx

diff --git a/Engine/Plugins/Messaging/UdpMessaging/Source/UdpMessaging/Private/Transport/UdpDeserializedMessage.cpp b/Engine/Plugins/Messaging/UdpMessaging/Source/UdpMessaging/Private/Transport/UdpDeserializedMessage.cpp
--- a/Engine/Plugins/Messaging/UdpMessaging/Source/UdpMessaging/Private/Transport/UdpDeserializedMessage.cpp
+++ b/Engine/Plugins/Messaging/UdpMessaging/Source/UdpMessaging/Private/Transport/UdpDeserializedMessage.cpp
@@ -150,6 +150,26 @@ const FDateTime& FUdpDeserializedMessage::GetTimeSent() const
 /* FUdpDeserializedMessageDetails implementation
 *****************************************************************************/
 
+TWeakObjectPtr<UScriptStruct> DeserializeSafeLookupForFindObject_InThread(FStringView TypeName)
+{
+	UUdpMessagingSettings const* CurrentSettings = GetDefault<UUdpMessagingSettings>();
+	const bool bRunningInTaskThread = (CurrentSettings->MaxConcurrentDeserializationTasks > 1);
+	if (bRunningInTaskThread)
+	{
+		// Since we are running in a task thread. We can just spin wait to get the type info.
+		while (UE::IsSavingPackage() || IsGarbageCollecting())
+		{
+			FPlatformProcess::SleepNoStats(0.01f);			
+		}
+	}
+	if (!UE::IsSavingPackage() && !IsGarbageCollecting())
+	{
+		return FindFirstObjectSafe<UScriptStruct>(TypeName, EFindFirstObjectOptions::EnsureIfAmbiguous);
+	}
+
+	return nullptr;
+}
+
 template <typename ReturnType>
 ReturnType FUdpDeserializedMessageDetails::GetMessageType(FMemoryReader& MessageReader)
 {
@@ -170,7 +190,7 @@ bool FUdpDeserializedMessageDetails::DeserializeV10(FUdpDeserializedMessage& Des
 
 		if (DeserializedMessage.TypeInfo == nullptr)
 		{
-			DeserializedMessage.TypeInfo = FindFirstObjectSafe<UScriptStruct>(*MessageType.ToString(), EFindFirstObjectOptions::EnsureIfAmbiguous);
+			DeserializedMessage.TypeInfo = DeserializeSafeLookupForFindObject_InThread(*MessageType.ToString());
 		}
 
 		if (!DeserializedMessage.TypeInfo.IsValid(false, true))
@@ -257,7 +277,7 @@ bool FUdpDeserializedMessageDetails::DeserializeV11_15(FUdpDeserializedMessage&
 
 		if (DeserializedMessage.TypeInfo == nullptr)
 		{
-			DeserializedMessage.TypeInfo = FindFirstObjectSafe<UScriptStruct>(*MessageType.ToString(), EFindFirstObjectOptions::EnsureIfAmbiguous);
+			DeserializedMessage.TypeInfo = DeserializeSafeLookupForFindObject_InThread(*MessageType.ToString());
 		}
 
 		if (!DeserializedMessage.TypeInfo.IsValid(false, true))
@@ -383,7 +403,7 @@ bool FUdpDeserializedMessageDetails::DeserializeV16_17(FUdpDeserializedMessage&
 
 		if (DeserializedMessage.TypeInfo == nullptr)
 		{
-			DeserializedMessage.TypeInfo = FindObjectSafe<UScriptStruct>(MessageType);
+			DeserializedMessage.TypeInfo = DeserializeSafeLookupForFindObject_InThread(*MessageType.ToString());
 		}
 
 		if (!DeserializedMessage.TypeInfo.IsValid(false, true))
