UdpMessaging: Merge recent improvements from internal project stream.

- Replaced use of TUdpCircularQueue with TRingBuffer to enable node work queues to resize if necessary; removed lots of dead code related to accommodating the "full" condition.
- Forked FUdpSocketSender to introduce a version that's more tolerant of transient or semi-expected send failures to arbitrary bus recipients.
- Moved FUdpMessageBeacon traffic onto the SocketSender to help ensure it's treated fairly.
- Fixed potential unsafe access to uninitialized segmenters.
- Made the send window size configurable using the cvar `MessageBus.UDP.MaxWindowSize`.
- Enable checking for expired inflight segments by default.

#FailedValidations Preflight

diff --git a/Engine/Plugins/Messaging/UdpMessaging/Source/UdpMessaging/Private/Shared/SocketSender.h b/Engine/Plugins/Messaging/UdpMessaging/Source/UdpMessaging/Private/Shared/SocketSender.h
new file mode 100644
--- /dev/null
+++ b/Engine/Plugins/Messaging/UdpMessaging/Source/UdpMessaging/Private/Shared/SocketSender.h
@@ -0,0 +1,296 @@
+// Copyright Epic Games, Inc. All Rights Reserved.
+
+#pragma once
+
+#include <atomic>
+#include "Containers/MpscQueue.h"
+#include "HAL/Event.h"
+#include "HAL/PlatformProcess.h"
+#include "HAL/Runnable.h"
+#include "HAL/RunnableThread.h"
+#include "Logging/StructuredLog.h"
+#include "Misc/SingleThreadRunnable.h"
+#include "Sockets.h"
+#include "SocketSubsystem.h"
+#include "UdpMessagingPrivate.h"
+
+#include "Interfaces/IPv4/IPv4Endpoint.h"
+
+
+namespace UE::UdpMessaging {
+
+/**
+ * Asynchronously sends data to a UDP socket.
+ */
+class FSocketSender
+	: public FRunnable
+	, private FSingleThreadRunnable
+{
+	// Structure for outbound packets.
+	struct FPacket
+	{
+		/** Holds the packet's data. */
+		TSharedPtr<TArray<uint8>> Data;
+
+		/** Holds the recipient. */
+		FIPv4Endpoint Recipient;
+
+		/** Default constructor. */
+		FPacket() { }
+
+		/** Creates and initializes a new instance. */
+		FPacket(const TSharedRef<TArray<uint8>>& InData, const FIPv4Endpoint& InRecipient)
+			: Data(InData)
+			, Recipient(InRecipient)
+		{ }
+	};
+
+	enum class EUpdateResult
+	{
+		/** The queued was emptied, and we can await new work. */
+		Done,
+
+		/** There was a transient failure, and we should attempt to service the queue again soon. */
+		Retry,
+
+		/** Tear down this FSocketSender and percolate the failure to the caller. */
+		Fatal,
+	};
+
+public:
+
+	struct FOptions
+	{
+		int32 SendBufferSize = 512 * 1024;
+		FTimespan WaitTime = FTimespan::FromMilliseconds(100);
+
+		FOptions()
+		{
+		}
+	};
+
+	/**
+	 * Creates and initializes a new socket sender.
+	 *
+	 * @param InSocket The UDP socket to use for sending data.
+	 * @param InDescription The description text (for debugging).
+	 */
+	FSocketSender(FSocket* InSocket, const TCHAR* InDescription, const FOptions& InOptions = FOptions())
+		: Socket(InSocket)
+		, SocketSubsystem(ISocketSubsystem::Get(PLATFORM_SOCKETSUBSYSTEM))
+		, bStopping(false)
+		, WaitTime(InOptions.WaitTime)
+		, Description(InDescription)
+	{
+		check(Socket != nullptr);
+		check(Socket->GetSocketType() == SOCKTYPE_Datagram);
+
+		int32 NewSize = 0;
+		Socket->SetSendBufferSize(InOptions.SendBufferSize, NewSize);
+
+		WorkEvent = FPlatformProcess::GetSynchEventFromPool();
+		Thread = FRunnableThread::Create(this, InDescription, 128 * 1024, TPri_AboveNormal, FPlatformAffinity::GetPoolThreadMask());
+	}
+
+	/** Virtual destructor. */
+	virtual ~FSocketSender()
+	{
+		if (Thread != nullptr)
+		{
+			Thread->Kill(true);
+			delete Thread;
+		}
+
+		FPlatformProcess::ReturnSynchEventToPool(WorkEvent);
+		WorkEvent = nullptr;
+	}
+
+public:
+
+	/**
+	 * Sends data to the specified recipient.
+	 *
+	 * @param Data The data to send.
+	 * @param Recipient The recipient.
+	 * @return true if the data will be sent, false otherwise.
+	 */
+	bool Send(const TSharedRef<TArray<uint8>, ESPMode::ThreadSafe>& Data, const FIPv4Endpoint& Recipient)
+	{
+		if (!bStopping)
+		{
+			SendQueue.Enqueue(Data, Recipient);
+			WorkEvent->Trigger();
+			return true;
+		}
+
+		return false;
+	}
+
+public:
+
+	//~ FRunnable interface
+
+	virtual FSingleThreadRunnable* GetSingleThreadInterface() override
+	{
+		return this;
+	}
+
+	virtual bool Init() override
+	{
+		return true;
+	}
+
+	virtual uint32 Run() override
+	{
+		while (!bStopping)
+		{
+			const EUpdateResult Result = Update(WaitTime);
+			switch (Result)
+			{
+				case EUpdateResult::Done:
+					WorkEvent->Wait(WaitTime);
+					continue;
+				case EUpdateResult::Retry:
+					continue;
+				case EUpdateResult::Fatal:
+					bStopping = true;
+					return 0;
+			}
+		}
+
+		return 0;
+	}
+
+	virtual void Stop() override
+	{
+		bStopping = true;
+		WorkEvent->Trigger();
+	}
+
+	virtual void Exit() override { }
+
+protected:
+
+	/**
+	 * Update this socket sender.
+	 *
+	 * @param Time to wait for the socket.
+	 * @return true on success, false otherwise.
+	 */
+	EUpdateResult Update(const FTimespan& SocketWaitTime)
+	{
+		while (!SendQueue.IsEmpty())
+		{
+			if (!Socket->Wait(ESocketWaitConditions::WaitForWrite, SocketWaitTime))
+			{
+				return EUpdateResult::Retry;
+			}
+
+			const FPacket* Packet = SendQueue.Peek();
+			const FIPv4Endpoint Recipient = Packet->Recipient;
+			const int32 PacketNumBytes = Packet->Data->Num();
+			int32 SentNumBytes = 0;
+			const bool bSendResult =
+				Socket->SendTo(Packet->Data->GetData(), PacketNumBytes, SentNumBytes, *Recipient.ToInternetAddr());
+
+			const ESocketErrors SocketError = SocketSubsystem->GetLastErrorCode();
+
+			if (!bSendResult)
+			{
+				if (SocketError == SE_EWOULDBLOCK || SocketError == SE_ENOBUFS || SocketError == SE_EINTR)
+				{
+					// Leave the packet in the queue to try again.
+					return EUpdateResult::Retry;
+				}
+			}
+
+			// Only pop after we've checked for "retry-able" error codes.
+			Packet = nullptr;
+			SendQueue.Dequeue();
+
+			if (!bSendResult)
+			{
+				bool bFatalError;
+				switch (SocketError)
+				{
+					case SE_ENETDOWN:
+					case SE_ENETUNREACH:
+					case SE_EHOSTDOWN:
+					case SE_EHOSTUNREACH:
+					case SE_EADDRNOTAVAIL:
+						bFatalError = false;
+						break;
+					default:
+						bFatalError = true;
+						break;
+				}
+
+				if (!bFatalError)
+				{
+					// Non-fatal; on to the next packet in the queue.
+					UE_LOGFMT(LogUdpMessaging, Verbose,
+						"Sender {Description}: SendTo failed (destination: {Recipient}) ({ErrorStr})",
+						Description, Recipient.ToString(), SocketSubsystem->GetSocketError(SocketError));
+					continue;
+				}
+				else
+				{
+					UE_LOGFMT(LogUdpMessaging, Error,
+						"Sender {Description}: SendTo failed (destination: {Recipient}) ({ErrorStr})",
+						Description, Recipient.ToString(), SocketSubsystem->GetSocketError(SocketError));
+				}
+			}
+
+			if (SentNumBytes != PacketNumBytes)
+			{
+				if (SentNumBytes >= 0)
+				{
+					// FIXME?: In the absence of another socket error, could be a retry? Is this even possible for UDP?
+					UE_LOGFMT(LogUdpMessaging, Error,
+						"Sender {Description}: Incomplete send (destination: {Recipient}) ({Sent}/{Total} bytes)",
+						Description, Recipient.ToString(), SentNumBytes, PacketNumBytes);
+				}
+				return EUpdateResult::Fatal;
+			}
+		}
+
+		return EUpdateResult::Done;
+	}
+
+protected:
+
+	//~ FSingleThreadRunnable interface
+
+	virtual void Tick() override
+	{
+		Update(FTimespan::Zero());
+	}
+
+private:
+
+	/** The send queue. */
+	TMpscQueue<FPacket> SendQueue;
+
+	/** The network socket. */
+	FSocket* Socket;
+
+	/** Cached pointer to the platform socket subsystem. */
+	ISocketSubsystem* SocketSubsystem;
+
+	/** Flag indicating that the thread is stopping. */
+	std::atomic<bool> bStopping;
+
+	/** The thread object. */
+	FRunnableThread* Thread;
+
+	/** Maximum time to wait for the socket to be ready to send. */
+	FTimespan WaitTime;
+
+	/** An event signaling that outbound messages need to be processed. */
+	FEvent* WorkEvent;
+
+	/** Description of this sender (for debugging). */
+	FString Description;
+};
+
+} // namespace UE::UdpMessaging
diff --git a/Engine/Plugins/Messaging/UdpMessaging/Source/UdpMessaging/Private/Transport/UdpMessageBeacon.cpp b/Engine/Plugins/Messaging/UdpMessaging/Source/UdpMessaging/Private/Transport/UdpMessageBeacon.cpp
--- a/Engine/Plugins/Messaging/UdpMessaging/Source/UdpMessaging/Private/Transport/UdpMessageBeacon.cpp
+++ b/Engine/Plugins/Messaging/UdpMessaging/Source/UdpMessaging/Private/Transport/UdpMessageBeacon.cpp
@@ -4,12 +4,12 @@
 #include "Transport/UdpMessageBeacon.h"
 #include "HAL/PlatformProcess.h"
 #include "UdpMessagingPrivate.h"
-#include "Common/UdpSocketSender.h"
 
 #include "HAL/Event.h"
 #include "HAL/RunnableThread.h"
 #include "Serialization/ArrayWriter.h"
 #include "Sockets.h"
+#include "Shared/SocketSender.h"
 #include "Shared/UdpMessageSegment.h"
 
 /* FUdpMessageHelloSender static initialization
@@ -22,17 +22,21 @@ const FTimespan FUdpMessageBeacon::MinimumInterval = FTimespan::FromMilliseconds
 /* FUdpMessageHelloSender structors
  *****************************************************************************/
 
-FUdpMessageBeacon::FUdpMessageBeacon(FSocket* InSocket, const FGuid& InSocketId, const FIPv4Endpoint& InMulticastEndpoint)
+FUdpMessageBeacon::FUdpMessageBeacon(
+	UE::UdpMessaging::FSocketSender* InSocketSender,
+	const FGuid& InSocketId,
+	const FIPv4Endpoint& InMulticastEndpoint
+)
 	: BeaconInterval(MinimumInterval)
 	, LastEndpointCount(1)
 	, LastHelloSent(FDateTime::MinValue())
 	, NextHelloTime(FDateTime::UtcNow())
 	, NodeId(InSocketId)
-	, Socket(InSocket)
+	, SocketSender(InSocketSender)
 	, Stopping(false)
 	, bSocketError(false)
 {
-	check(Socket);
+	check(SocketSender);
 	EndpointLeftEvent = FPlatformProcess::GetSynchEventFromPool(false);
 	MulticastAddress = InMulticastEndpoint.ToInternetAddr();
 	AddLocalEndpoints();
@@ -145,23 +149,15 @@ bool FUdpMessageBeacon::SendSegment(EUdpMessageSegments SegmentType, const FTime
 		Header.SegmentType = SegmentType;
 	}
 
-	FArrayWriter Writer;
-	{
-		Writer << Header;
-		Writer << NodeId;
-	}
-
-	int32 Sent;
+	TSharedRef<FArrayWriter, ESPMode::ThreadSafe> Writer = MakeShared<FArrayWriter, ESPMode::ThreadSafe>();
 
-	if (!Socket->Wait(ESocketWaitConditions::WaitForWrite, SocketWaitTime))
-	{
-		return false; // socket not ready for sending
-	}
+	*Writer << Header;
+	*Writer << NodeId;
 
-	if (!Socket->SendTo(Writer.GetData(), Writer.Num(), Sent, *MulticastAddress))
+	if (!SocketSender->Send(Writer, MulticastAddress))
 	{
 		bSocketError = true;
-		return false; // send failed
+		return false;
 	}
 
 	return true;
@@ -179,30 +175,22 @@ bool FUdpMessageBeacon::SendPing(const FTimespan& SocketWaitTime)
 	}
 	uint8 ActualProtocolVersion = UDP_MESSAGING_TRANSPORT_PROTOCOL_VERSION;
 
-	
-	FArrayWriter Writer;
+	TSharedRef<FArrayWriter, ESPMode::ThreadSafe> Writer = MakeShared<FArrayWriter, ESPMode::ThreadSafe>();
 	{
-		Writer << Header;
-		Writer << NodeId;
-		Writer << ActualProtocolVersion; // Send our actual Protocol version as part of the ping message
+		*Writer << Header;
+		*Writer << NodeId;
+		*Writer << ActualProtocolVersion; // Send our actual Protocol version as part of the ping message
 	}
 
-	if (!Socket->Wait(ESocketWaitConditions::WaitForWrite, SocketWaitTime))
-	{
-		return false; // socket not ready for sending
-	}
 
-	int32 Sent;
 	bool Result = true;
 	for (const auto& StaticAddress : StaticAddresses)
 	{
-		if (!Socket->SendTo(Writer.GetData(), Writer.Num(), Sent, *StaticAddress))
+		if (!SocketSender->Send(Writer, StaticAddress))
 		{	
 			Result = false; // send failed
 		}
-
 	}
-
 	return Result;
 }
 
diff --git a/Engine/Plugins/Messaging/UdpMessaging/Source/UdpMessaging/Private/Transport/UdpMessageBeacon.h b/Engine/Plugins/Messaging/UdpMessaging/Source/UdpMessaging/Private/Transport/UdpMessageBeacon.h
--- a/Engine/Plugins/Messaging/UdpMessaging/Source/UdpMessaging/Private/Transport/UdpMessageBeacon.h
+++ b/Engine/Plugins/Messaging/UdpMessaging/Source/UdpMessaging/Private/Transport/UdpMessageBeacon.h
@@ -14,7 +14,9 @@ enum class EUdpMessageSegments : uint8;
 class FEvent;
 class FInternetAddr;
 class FSocket;
-class FUdpSocketSender;
+
+namespace UE::UdpMessaging { class FSocketSender; }
+
 
 /**
  * Implements a beacon sender thread.
@@ -29,12 +31,16 @@ public:
 	/** 
 	 * Creates and initializes a new beacon sender.
 	 *
-	 * @param InSocket The network socket used to send Hello segments.
+	 * @param InSocketSender The network socket sender used to send Hello segments.
 	 * @param InSocketId The network socket identifier (used to detect unicast endpoint).
 	 * @param InMulticastEndpoint The multicast group endpoint to transport messages to.
 	 * @param InStaticEndpoints The static nodes to broadcast to alongside the multicast.
 	 */
-	FUdpMessageBeacon(FSocket* InSocketSender, const FGuid& InSocketId, const FIPv4Endpoint& InMulticastEndpoint);
+	FUdpMessageBeacon(
+		UE::UdpMessaging::FSocketSender* InSocketSender,
+		const FGuid& InSocketId,
+		const FIPv4Endpoint& InMulticastEndpoint
+	);
 
 	/** Virtual destructor. */
 	virtual ~FUdpMessageBeacon();
@@ -158,8 +164,8 @@ private:
 	/** Holds local node identifier. */
 	FGuid NodeId;
 
-	/** Holds the socket used to send Hello segments. */
-	FSocket* Socket;
+	/** Holds the socket sender used to send Hello segments. */
+	UE::UdpMessaging::FSocketSender* SocketSender;
 
 	/** Holds a flag indicating that the thread is stopping. */
 	bool Stopping;
diff --git a/Engine/Plugins/Messaging/UdpMessaging/Source/UdpMessaging/Private/Transport/UdpMessageProcessor.cpp b/Engine/Plugins/Messaging/UdpMessaging/Source/UdpMessaging/Private/Transport/UdpMessageProcessor.cpp
--- a/Engine/Plugins/Messaging/UdpMessaging/Source/UdpMessaging/Private/Transport/UdpMessageProcessor.cpp
+++ b/Engine/Plugins/Messaging/UdpMessaging/Source/UdpMessaging/Private/Transport/UdpMessageProcessor.cpp
@@ -3,16 +3,15 @@
 #include "Transport/UdpMessageProcessor.h"
 #include "Algo/AllOf.h"
 
-#include "Common/UdpSocketSender.h"
 #include "CoreGlobals.h"
 #include "Misc/ConfigCacheIni.h"
 #include "Serialization/ArrayReader.h"
 #include "Serialization/ArrayWriter.h"
 
 #include "Transport/UdpDeserializedMessage.h"
-#include "UdpCircularQueue.h"
 #include "UdpMessagingPrivate.h"
 #include "UdpMessagingTracing.h"
+#include "Shared/SocketSender.h"
 #include "Shared/UdpMessagingSettings.h"
 #include "Transport/UdpMessageBeacon.h"
 #include "Transport/UdpMessageSegmenter.h"
@@ -45,10 +44,16 @@ TAutoConsoleVariable<FString> CVarConnectionsToError(
 	TEXT("This can be a comma separated list in the form IPAddr2:port,IPAddr3:port"),
 	ECVF_Default);
 
-TAutoConsoleVariable<int32> CVarCheckForExpiredWithFullQueue(
-	TEXT("MessageBus.UDP.CheckForExpiredWithFullQueue"),
-	0,
-	TEXT("Attempts to release pressure on the work queue by checking if inflight segments have expired with no acknowledgement.\n"),
+TAutoConsoleVariable<int32> CVarCheckForExpired(
+	TEXT("MessageBus.UDP.CheckForExpired"),
+	1,
+	TEXT("Attempts to relieve pressure on the work queue by checking if inflight segments have expired with no acknowledgement.\n"),
+	ECVF_Default);
+
+TAutoConsoleVariable<int32> CVarUdpMessagingMaxWindowSize(
+	TEXT("MessageBus.UDP.MaxWindowSize"),
+	2048,
+	TEXT("Maximum window size for sending to endpoints.\n"),
 	ECVF_Default);
 
 
@@ -107,7 +112,7 @@ FUdpMessageProcessor::FNodeInfo::FNodeInfo()
 	, ReliableWorkQueue(UE::Private::MessageProcessor::GetMessageProcessorWorkQueueSize())
 	, UnreliableWorkQueue(UE::Private::MessageProcessor::GetMessageProcessorWorkQueueSize())
 {
-	ComputeWindowSize(0,0);
+	ComputeWindowSize(0, 0);
 }
 
 FUdpMessageProcessor::FUdpMessageProcessor(FSocket& InSocket, const FGuid& InNodeId, const FIPv4Endpoint& InMulticastEndpoint)
@@ -343,9 +348,9 @@ bool FUdpMessageProcessor::Init()
 {
 	if (!bIsInitialized)
 	{
-		SocketSender = MakeUnique<FUdpSocketSender>(Socket, TEXT("FUdpMessageProcessor.Sender"));
+		SocketSender = MakeUnique<UE::UdpMessaging::FSocketSender>(Socket, TEXT("FUdpMessageProcessor.Sender"));
 
-		Beacon = MakeUnique<FUdpMessageBeacon>(Socket, LocalNodeId, MulticastEndpoint);
+		Beacon = MakeUnique<FUdpMessageBeacon>(SocketSender.Get(), LocalNodeId, MulticastEndpoint);
 
 		// Current protocol version 17
 		SupportedProtocolVersions.Add(UDP_MESSAGING_TRANSPORT_PROTOCOL_VERSION);
@@ -469,6 +474,7 @@ void FUdpMessageProcessor::AcknowledgeReceipt(int32 MessageId, const FNodeInfo&
 		UE_LOG(LogUdpMessaging, Warning, TEXT("FUdpMessageProcessor::AcknowledgeReceipt failed to send %s."), *NodeInfo.Endpoint.ToString());
 		return;
 	}
+
 	UE_LOG(LogUdpMessaging, Verbose, TEXT("Sending EUdpMessageSegments::Acknowledge for msg %d from %s"), MessageId, *NodeInfo.NodeId.ToString());
 }
 
@@ -490,6 +496,7 @@ void FUdpMessageProcessor::ConsumeInboundSegments()
 	FInboundSegment Segment;
 
 	FTimespan MaxWorkTimespan = ComputeMaxWorkTimespan(DeadHelloIntervals, Beacon.Get());
+
 	while (InboundSegments.Dequeue(Segment))
 	{
 		// quick hack for TTP# 247103
@@ -522,7 +529,8 @@ void FUdpMessageProcessor::ConsumeInboundSegments()
 
 			NodeInfo.Endpoint = Segment.Sender;
 			NodeInfo.LastSegmentReceivedTime = CurrentTime;
-
+			NodeInfo.MaxWindowSize = CVarUdpMessagingMaxWindowSize.GetValueOnAnyThread();
+			
 			switch (Header.SegmentType)
 			{
 			case EUdpMessageSegments::Abort:
@@ -581,7 +589,7 @@ void FUdpMessageProcessor::ConsumeInboundSegments()
 	}
 }
 
-bool FUdpMessageProcessor::ConsumeOneOutboundMessage(const FOutboundMessage& OutboundMessage)
+void FUdpMessageProcessor::ConsumeOneOutboundMessage(const FOutboundMessage& OutboundMessage)
 {
 	const auto FindNodeWithLog = [this](const FGuid &Id)
 	{
@@ -598,52 +606,41 @@ bool FUdpMessageProcessor::ConsumeOneOutboundMessage(const FOutboundMessage& Out
 					  [FindNodeWithLog](const FGuid &Id) {return FindNodeWithLog(Id);},
 					  [this](const FGuid &Id) {return KnownNodes.Find(Id);} );
 
-	
-	const bool bCanConsume = Algo::AllOf(Recipients, [this](FNodeInfo *Node)
-		{
-			return Node->CanCommitToWorkQueue(CurrentTime);
-		});
-
 	const bool bIsReliable = EnumHasAnyFlags(OutboundMessage.MessageFlags, EMessageFlags::Reliable);
-	if (bCanConsume)
+
+	++LastSentMessage;
+	if (LastSentMessage < 0)
 	{
-		++LastSentMessage;
-		if (LastSentMessage < 0)
+		// Prevent negative message ids
+		//
+		LastSentMessage = 0;
+	}
+	for (FNodeInfo* RecipientNodeInfo : Recipients)
+	{
+		UE_LOG(LogUdpMessaging, Verbose, TEXT("Passing %" INT64_FMT " byte message to be segment-sent %d to %s with id %s"),
+				OutboundMessage.SerializedMessage->TotalSize(),
+				LastSentMessage,
+				*RecipientNodeInfo->Endpoint.ToString(),
+				*RecipientNodeInfo->NodeId.ToString());
+
+		if (!bIsReliable && !RecipientNodeInfo->CanSendSegments())
 		{
-			// Prevent negative message ids
-			//
-			LastSentMessage = 0;
+			// Discard unreliable messages that cannot be sent.
+			continue;
 		}
-		for (FNodeInfo *RecipientNodeInfo : Recipients)
-		{
-			UE_LOG(LogUdpMessaging, Verbose, TEXT("Passing %" INT64_FMT " byte message to be segment-sent %d to %s with id %s"),
-				   OutboundMessage.SerializedMessage->TotalSize(),
-				   LastSentMessage,
-				   *RecipientNodeInfo->Endpoint.ToString(),
-				   *RecipientNodeInfo->NodeId.ToString());
 
-			if (!bIsReliable && !RecipientNodeInfo->CanSendSegments())
-			{
-				// Discard unreliable messages that cannot be sent.
-				continue;
-			}
-
-			TSharedPtr<FUdpMessageSegmenter> Segmenter = MakeShared<FUdpMessageSegmenter>(OutboundMessage.SerializedMessage.ToSharedRef(), UDP_MESSAGING_SEGMENT_SIZE);
-			RecipientNodeInfo->Segmenters.Add(
-				LastSentMessage, Segmenter);
-			if (bIsReliable)
-			{
-				RecipientNodeInfo->ReliableWorkQueue.Enqueue(LastSentMessage);
-			}
-			else
-			{
-				RecipientNodeInfo->UnreliableWorkQueue.Enqueue(LastSentMessage);
-			}
+		TSharedPtr<FUdpMessageSegmenter> Segmenter = MakeShared<FUdpMessageSegmenter>(OutboundMessage.SerializedMessage.ToSharedRef(), UDP_MESSAGING_SEGMENT_SIZE);
+		RecipientNodeInfo->Segmenters.Add(
+			LastSentMessage, Segmenter);
+		if (bIsReliable)
+		{
+			RecipientNodeInfo->ReliableWorkQueue.Add(LastSentMessage);
+		}
+		else
+		{
+			RecipientNodeInfo->UnreliableWorkQueue.Add(LastSentMessage);
 		}
 	}
-
-	// Return true that it was consumed or not reliable. Unreliable messages can be discarded.
-	return bCanConsume || !bIsReliable;
 }
 
 void FUdpMessageProcessor::ConsumeOutboundMessages()
@@ -651,26 +648,9 @@ void FUdpMessageProcessor::ConsumeOutboundMessages()
 	SCOPED_MESSAGING_TRACE(FUdpMessageProcessor_ConsumeOutputMessages);
 	
 	FOutboundMessage OutboundMessage;
-
-	if (DeferredOutboundMessage)
-	{
-		if (ConsumeOneOutboundMessage(DeferredOutboundMessage.GetValue()))
-		{
-			DeferredOutboundMessage.Reset();
-		}
-		else
-		{
-			return ;
-		}
-	}
-
 	while (OutboundMessages.Dequeue(OutboundMessage))
 	{
-		if (!ConsumeOneOutboundMessage(OutboundMessage))
-		{
-			DeferredOutboundMessage = OutboundMessage;
-			return ;
-		}
+		ConsumeOneOutboundMessage(OutboundMessage);
 	}
 }
 
@@ -871,7 +851,7 @@ void FUdpMessageProcessor::ProcessPingSegment(FInboundSegment& Segment, FNodeInf
 
 	if (!SocketSender->Send(Writer, NodeInfo.Endpoint))
 	{
-		UE_LOG(LogUdpMessaging, Warning, TEXT("FUdpMessageProcessor::ProcessPingSegment failed to send ping segment %s."), *NodeInfo.Endpoint.ToString());		
+		UE_LOG(LogUdpMessaging, Warning, TEXT("FUdpMessageProcessor::ProcessPingSegment failed to send ping segment %s."), *NodeInfo.Endpoint.ToString());
 	}
 }
 
@@ -1085,6 +1065,7 @@ void FUdpMessageProcessor::SendKnownNodesToKnownNodes()
 
 		OutEndpoints = SlicedEndpointView;
 		OutIds = SlicedIdView;
+
 		UE_LOG(LogUdpMessaging, VeryVerbose, TEXT("FUdpMessageProcessor::SendKnownNodesToKnownNodes Sending updated known nodes."));
 
 		for (const FIPv4Endpoint& Endpoint : Endpoints)
@@ -1095,6 +1076,7 @@ void FUdpMessageProcessor::SendKnownNodesToKnownNodes()
 				*Writer << OutEndpoints;
 				*Writer << OutIds;
 			}
+
 			if (!SocketSender->Send(Writer, Endpoint))
 			{
 				UE_LOG(LogUdpMessaging, Warning, TEXT("FUdpMessageProcessor::SendKnownNodesToKnownNodes failed to share endpoint information to %s."), *Endpoint.ToString());
@@ -1307,10 +1289,15 @@ FSentSegmentInfo FUdpMessageProcessor::SendNextSegmentForMessageId(FNodeInfo& No
 	}
 	else
 	{
-		SentInfo.bIsTraced = EnumHasAnyFlags(Segmenter->GetMessageFlags(), EMessageFlags::Traced);
-		SentInfo.bIsReliable = EnumHasAnyFlags(Segmenter->GetMessageFlags(), EMessageFlags::Reliable);
+		const bool bIsReady = Segmenter->IsInitialized(); 
+		if (bIsReady)
+		{
+			SentInfo.bIsTraced = EnumHasAnyFlags(Segmenter->GetMessageFlags(), EMessageFlags::Traced);
+			SentInfo.bIsReliable = EnumHasAnyFlags(Segmenter->GetMessageFlags(), EMessageFlags::Reliable);
+			SentInfo.bFullySent = Segmenter->IsSendingComplete();
+		}
 		SentInfo.bRequiresRequeue = true;
-		SentInfo.bFullySent = Segmenter->IsInitialized() && Segmenter->IsSendingComplete();
+		SentInfo.bIsReady = bIsReady;
 	}
 
 	NodeInfo.Statistics.TotalBytesSent += SentInfo.BytesSent;
@@ -1318,7 +1305,7 @@ FSentSegmentInfo FUdpMessageProcessor::SendNextSegmentForMessageId(FNodeInfo& No
 	return MoveTemp(SentInfo);
 }
 
-TOptional<FSentSegmentInfoArray> FUdpMessageProcessor::DrainQueue(FNodeInfo& NodeInfo, TUdpCircularQueue<int32>& WorkQueue, int64 DrainCount, const bool bDrainOverflow)
+TOptional<FSentSegmentInfoArray> FUdpMessageProcessor::DrainQueue(FNodeInfo& NodeInfo, FNodeInfo::FWorkQueue& WorkQueue, int64 DrainCount, const bool bDrainOverflow)
 {
 	FUdpMessageSegment::FHeader Header
 	{
@@ -1345,10 +1332,10 @@ TOptional<FSentSegmentInfoArray> FUdpMessageProcessor::DrainQueue(FNodeInfo& Nod
 		}
 	}
 	
-	int32 MessageId;
 	// Process messages in the work queue.
-	while (DrainCount > 0 && NodeInfo.CanSendSegments() && WorkQueue.Dequeue(MessageId))
+	while (DrainCount > 0 && NodeInfo.CanSendSegments() && !WorkQueue.IsEmpty())
 	{
+		int32 MessageId = WorkQueue.PopFrontValue();
 		FSentSegmentInfo Info = SendNextSegmentForMessageId(NodeInfo, Header, MessageId);
 		if (Info.bSendSocketError)
 		{
@@ -1360,19 +1347,18 @@ TOptional<FSentSegmentInfoArray> FUdpMessageProcessor::DrainQueue(FNodeInfo& Nod
 	return ProcessedMessages;
 }
 
-int32 FUdpMessageProcessor::RequeueWork(FNodeInfo& NodeInfo, TUdpCircularQueue<int32>& WorkQueue, const FSentSegmentInfoArray& ProcessedMessages)
+int32 FUdpMessageProcessor::RequeueWork(FNodeInfo& NodeInfo, FNodeInfo::FWorkQueue& WorkQueue, const FSentSegmentInfoArray& ProcessedMessages)
 {
 	int32 BytesSent = 0;
 	for (const FSentSegmentInfo& Info : ProcessedMessages)
 	{
-		if ((Info.bFullySent && Info.bRequiresRequeue) ||
-			(Info.bRequiresRequeue && WorkQueue.IsFull()) )
+		if (Info.bRequiresRequeue && Info.bFullySent)
 		{
 			NodeInfo.OverflowForPendingAck.Add(Info.MessageId);
 		}
 		else if (Info.bRequiresRequeue)
 		{
-			WorkQueue.Enqueue(Info.MessageId);
+			WorkQueue.Add(Info.MessageId);
 		}
 
 		BytesSent += Info.BytesSent;
@@ -1380,12 +1366,12 @@ int32 FUdpMessageProcessor::RequeueWork(FNodeInfo& NodeInfo, TUdpCircularQueue<i
 	return BytesSent;	
 }
 
-bool FUdpMessageProcessor::ProcessQueue(FNodeInfo& NodeInfo, TUdpCircularQueue<int32>& WorkQueue, int64 DrainLimit, int32& BytesSent, bool bShouldDrainOverflowQueue)
+bool FUdpMessageProcessor::ProcessQueue(FNodeInfo& NodeInfo, FNodeInfo::FWorkQueue& WorkQueue, int64 DrainLimit, int32& BytesSent, bool bShouldDrainOverflowQueue)
 {
 	int64 SendCapacity = NodeInfo.SendCapacity();
 	int64 MaxSend = SendCapacity - DrainLimit;
 
-	while (!WorkQueue.IsEmpty() && NodeInfo.CanSendSegments() && SendCapacity > DrainLimit )
+	while (!WorkQueue.IsEmpty() && NodeInfo.CanSendSegments() && SendCapacity > DrainLimit)
 	{
 		if (TOptional<FSentSegmentInfoArray> ProcessedMessages = DrainQueue(NodeInfo, WorkQueue, MaxSend, bShouldDrainOverflowQueue))
 		{
@@ -1440,11 +1426,10 @@ int32 FUdpMessageProcessor::UpdateSegmenters(FNodeInfo& NodeInfo)
 		}
 	}
 	
-	if (CVarCheckForExpiredWithFullQueue.GetValueOnAnyThread() > 0 && !NodeInfo.CanSendSegments() && NodeInfo.ReliableWorkQueue.IsFull())
+	if (CVarCheckForExpired.GetValueOnAnyThread() > 0 && !NodeInfo.CanSendSegments())
 	{
-		UE_LOG(LogUdpMessaging, Warning, TEXT("Work queue is full sending to node at address %s. We cannot send new data. Attempting to expire old inflight segments."), *NodeInfo.Endpoint.ToString());
-		// RemoveLostSegments will clean-up the tracking map that we use to determine what is "inflight". Old data will be expired and scheduled
-		// for resend.
+		// RemoveLostSegments will clean-up the tracking map that we use to determine what is "inflight".
+		// Old data will be expired and scheduled for resend.
 		NodeInfo.RemoveLostSegments(CurrentTime);
 	}
 
diff --git a/Engine/Plugins/Messaging/UdpMessaging/Source/UdpMessaging/Private/Transport/UdpMessageProcessor.h b/Engine/Plugins/Messaging/UdpMessaging/Source/UdpMessaging/Private/Transport/UdpMessageProcessor.h
--- a/Engine/Plugins/Messaging/UdpMessaging/Source/UdpMessaging/Private/Transport/UdpMessageProcessor.h
+++ b/Engine/Plugins/Messaging/UdpMessaging/Source/UdpMessaging/Private/Transport/UdpMessageProcessor.h
@@ -3,6 +3,7 @@
 #pragma once
 
 #include "Containers/MpscQueue.h"
+#include "Containers/RingBuffer.h"
 
 #include "HAL/Platform.h"
 #include "ModuleDescriptor.h"
@@ -19,11 +20,12 @@
 
 #include "UdpMessageSegmenter.h"
 #include "UdpMessagingPrivate.h"
-#include "Transport/UdpCircularQueue.h"
+
 
 class IMessageContext;
 namespace FUdpMessageSegment { struct FDataChunk; }
 namespace FUdpMessageSegment { struct FHeader; }
+namespace UE::UdpMessaging { class FSocketSender; }
 
 class FArrayReader;
 class FEvent;
@@ -33,29 +35,11 @@ class FUdpMessageBeacon;
 class FUdpMessageSegmenter;
 class FUdpReassembledMessage;
 class FUdpSerializedMessage;
-class FUdpSocketSender;
 class IMessageAttachment;
 enum class EUdpMessageFormat : uint8;
 enum class EUdpMessageSegments : uint8;
 using FUdpReassembledMessageSharedPtr = TSharedPtr<FUdpReassembledMessage, ESPMode::ThreadSafe>;
 
-/**
- * Running statistics as known by the UdpMessageProcessor. This will be per endpoint.
- */
-struct FUdpMessageTransportStatistics
-{
-	uint64 BytesSent    = 0;
-	uint64 SegmentsLost = 0;
-	uint64 AcksReceived = 0;
-	uint64 SegmentsSent = 0;
-	uint64 SegmentsReceived = 0;
-
-	int32 SegmentsInFlight = 0;
-	uint32 WindowSize = 0;
-
-	FTimespan AverageRTT{0};
-	FIPv4Endpoint IpAddress;
-};
 
 /** Holds the current statistics for a given message segmenter. */
 struct FUdpSegmenterStats
@@ -97,6 +81,7 @@ struct FSentSegmentInfo
 	bool bSendSocketError = false;
 	bool bFullySent = false;
 	bool bIsTraced = false;
+	bool bIsReady = false;
 
 	/** Converts FSentSegmentInfo into a FSentData struct. */
 	FSentData AsSentData(int32 InMessageId, uint32 SegmentNumber, const FDateTime& CurrentTime)
@@ -132,6 +117,16 @@ FOnInboundTransferDataUpdated &OnReassemblerUpdated();
 
 }
 
+/**
+ * Heap allocator to avoid range checks for performance reasons.
+ */
+class FUdpMessagingRangeChecklessHeapAllocator : public FHeapAllocator
+{
+public:
+	/** Don't want to lose performance on range checks in performance-critical code. */
+	enum { RequireRangeCheck = false };
+};
+
 /**
  * Implements a message processor for UDP messages.
  */
@@ -160,11 +155,13 @@ class FUdpMessageProcessor
 		/** Holds the collection of message segmenters. */
 		TMap<int32, TSharedPtr<FUdpMessageSegmenter>> Segmenters;
 
+		using FWorkQueue = TRingBuffer<int32, FUdpMessagingRangeChecklessHeapAllocator>;
+
 		/** Holds of queue of MessageIds to send for reliable messages. They are processed in round-robin fashion. */
-		TUdpCircularQueue<int32> ReliableWorkQueue;
+		FWorkQueue ReliableWorkQueue;
 
 		/** Holds of queue of MessageIds to send for unreliable messages. They are processed in round-robin fashion. */
-		TUdpCircularQueue<int32> UnreliableWorkQueue;
+		FWorkQueue UnreliableWorkQueue;
 
 		/* Holds the list of messages that are still pending acknowledgement or overflow if overcommitted on work queue. */
 		TArray<int32> OverflowForPendingAck;
@@ -187,39 +184,15 @@ class FUdpMessageProcessor
 		/** Maximum number of segments we can have in flight. Start with a middle of the road value and adjust based on network performance. */
 		uint64 WindowSize = 1024;
 
+		/** Maximum window size that can be set. */
+		uint64 MaxWindowSize = 2048;
+
 		/** Various transport statistics for this endpoint */
 		FMessageTransportStatistics Statistics;
 
 		/** Default constructor. */
 		FNodeInfo();
 
-		/**
-		 * Check the work queue buffer and note a warning if we are overcommitted.
-		 */
-		bool CanCommitToWorkQueue(const FDateTime& InCurrentTime)
-		{
-			const bool bWorkQueueFull = ReliableWorkQueue.IsFull();
-			if (bWorkQueueFull)
-			{
-				const double ReportingIntervalInSeconds = 1;
-				double Seconds = (InCurrentTime - LastWorkQueueFullMessage).GetTotalSeconds();
-				if (Seconds > ReportingIntervalInSeconds)
-				{
-					SecondsInFullQueueState += Seconds;
-					UE_LOG(LogUdpMessaging, Warning,
-						   TEXT("Work queue for node %s is full. Send queue has been congested for %.3f seconds."),
-						   *Endpoint.ToString(), SecondsInFullQueueState);
-					LastWorkQueueFullMessage = InCurrentTime;
-				}
-			}
-			else
-			{
-				LastWorkQueueFullMessage = InCurrentTime;
-				SecondsInFullQueueState = 0;
-			}
-			return !bWorkQueueFull;
-		}
-
 		/**
 		 * Remap a MessageId and SegmentId pair into a 64 bit number using Szudzik pairing calculation.
 		 */
@@ -304,7 +277,7 @@ class FUdpMessageProcessor
 		void ComputeWindowSize(uint32 NumAcks, uint32 SegmentLoss)
 		{
 			const uint64 MinimumWindowSize = 64;
-			const uint64 MaximumWindowSize = 2048;
+			const uint64 MaximumWindowSize = MaxWindowSize;
 			if (SegmentLoss == 0)
 			{
 				// If we did not get any packet loss increase our window size.
@@ -877,10 +850,38 @@ protected:
 
 	virtual void Tick() override;
 
-	/** Return a copy of the known nodes as seen by the processor. This is used by debugging logging utils and should not be relied on as an API*/
-	TMap<FGuid, FNodeInfo> GetKnownNodesAsSnapshot()
+	/* A shallow copy of NodeInfo for use by the logger. */
+	struct FShallowNodeInfo
 	{
-		return KnownNodes;
+		/** Holds the node's IP endpoint. */
+		FIPv4Endpoint Endpoint;
+
+		/** Holds the time at which the last Hello was received. */
+		FDateTime LastSegmentReceivedTime;
+
+		/** Holds the endpoint's node identifier. */
+		FGuid NodeId;
+
+		/** Holds the protocol version this node is communicating with */
+		uint8 ProtocolVersion;
+	};
+
+	/** Return a shallow copy of the known nodes as seen by the processor. This is used and accessible by debugging logging utils. */
+	TArray<FShallowNodeInfo> GetKnownNodesAsSnapshot()
+	{
+		TArray<FShallowNodeInfo> Result;
+		for (auto& KnownNodePair : KnownNodes)
+		{
+			FNodeInfo& NodeInfo = KnownNodePair.Value;
+			Result.Add(
+				{
+					NodeInfo.Endpoint,
+					NodeInfo.LastSegmentReceivedTime,
+					NodeInfo.NodeId,
+					NodeInfo.ProtocolVersion
+			});
+		}
+		return Result;
 	}
 
 private:
@@ -891,19 +892,19 @@ private:
 	 * Drain a given send queue to the drain limit specified on the parameter listed. Also drain the overflow bucket if requested. Overflow
 	 * may occasionally fill from pending acks or work queue overflowing.  
 	 */
-	TOptional<FSentSegmentInfoArray> DrainQueue(FNodeInfo& NodeInfo, TUdpCircularQueue<int32>& WorkQueue, int64 DrainCount, const bool bDrainOverflow);
+	TOptional<FSentSegmentInfoArray> DrainQueue(FNodeInfo& NodeInfo, FNodeInfo::FWorkQueue& WorkQueue, int64 DrainCount, const bool bDrainOverflow);
 
 	/**
 	 * Refill the given queue with new work based on the processed messages. They will be put back in the queue to be processed later
 	 * in the next round to ensure fairness in the sending of data.
 	 */
-	int32 RequeueWork(FNodeInfo& NodeInfo, TUdpCircularQueue<int32>& WorkQueue, const FSentSegmentInfoArray& ProcessedMessages);
+	int32 RequeueWork(FNodeInfo& NodeInfo, FNodeInfo::FWorkQueue& WorkQueue, const FSentSegmentInfoArray& ProcessedMessages);
 
 	/**
 	 * Process a given queue.  This calls into DrainQueue and RequeueWork repeatedly until we reach the drain limit which is the minimum
 	 * value of of send capacity.  For example, if DrainLimit is 0 then we work until the send capacity is 0. 
 	 */
-	bool ProcessQueue(FNodeInfo& NodeInfo, TUdpCircularQueue<int32>& WorkQueue, int64 DrainLimit, int32& BytesSent, bool bShouldDrainOverflowQueue = false);
+	bool ProcessQueue(FNodeInfo& NodeInfo, FNodeInfo::FWorkQueue& WorkQueue, int64 DrainLimit, int32& BytesSent, bool bShouldDrainOverflowQueue = false);
 		
 	/** Send our list of known nodes to the known nodes. */
 	void SendKnownNodesToKnownNodes();
@@ -923,8 +924,8 @@ private:
 	/** Checks all known nodes to see if any segmenters have NeedSending set to true. */
 	bool MoreToSend();
 
-	/** Consume an Outbound Message for processing. A result of true will be returned if it was added for processing. */
-	bool ConsumeOneOutboundMessage(const FOutboundMessage& OutboundMessage);
+	/** Consume an Outbound Message for processing. */
+	void ConsumeOneOutboundMessage(const FOutboundMessage& OutboundMessage);
 
 	/** Holds the queue of outbound messages. */
 	TMpscQueue<FInboundSegment> InboundSegments;
@@ -975,7 +976,7 @@ private:
 	FSocket* Socket;
 
 	/** Holds the socket sender.*/
-	TUniquePtr<FUdpSocketSender> SocketSender;
+	TUniquePtr<UE::UdpMessaging::FSocketSender> SocketSender;
 
 	/** Enable sharing of our known nodes. */
 	bool bShareKnownNodes = false;
@@ -1019,9 +1020,6 @@ private:
 	/** The configured message format (from UUdpMessagingSettings). */
 	EUdpMessageFormat MessageFormat;
 
-	/** If our round-robin work queues couldn't accept the last outbound message then store a deferred message for next round. */
-	TOptional<FOutboundMessage> DeferredOutboundMessage;
-
 	/** Stores a cache of discovered UScriptStruct types found during message processing. */
 	TMap<FString, TWeakObjectPtr<UScriptStruct>> CachedTypeInfoMap;
 
diff --git a/Engine/Plugins/Messaging/UdpMessaging/Source/UdpMessaging/Private/Transport/UdpMessageProcessorLogger.cpp b/Engine/Plugins/Messaging/UdpMessaging/Source/UdpMessaging/Private/Transport/UdpMessageProcessorLogger.cpp
--- a/Engine/Plugins/Messaging/UdpMessaging/Source/UdpMessaging/Private/Transport/UdpMessageProcessorLogger.cpp
+++ b/Engine/Plugins/Messaging/UdpMessaging/Source/UdpMessaging/Private/Transport/UdpMessageProcessorLogger.cpp
@@ -598,14 +598,14 @@ void FUdpMessageProcessorLogger::DumpKnownNodeInfo()
 {
 	auto SnapshotStats = [this](FUdpMessageProcessorLoggerTyped<FNodeInfoStats> & NodeInfo)
 	{
-		TMap<FGuid, FUdpMessageProcessor::FNodeInfo> Snapshot = Processor->GetKnownNodesAsSnapshot();
-		for (auto It = Snapshot.CreateIterator(); It; ++It)
+		TArray<FUdpMessageProcessor::FShallowNodeInfo> Snapshot = Processor->GetKnownNodesAsSnapshot();
+		for (const FUdpMessageProcessor::FShallowNodeInfo& Value : Snapshot)
 		{
 			FNodeInfoStats Stat;
-			Stat.NodeId = It->Value.NodeId;
-			Stat.Endpoint = It->Value.Endpoint;
-			Stat.LastSegmentReceivedTime = It->Value.LastSegmentReceivedTime;
-			Stat.ProtocolVersion = It->Value.ProtocolVersion;
+			Stat.NodeId = Value.NodeId;
+			Stat.Endpoint = Value.Endpoint;
+			Stat.LastSegmentReceivedTime = Value.LastSegmentReceivedTime;
+			Stat.ProtocolVersion = Value.ProtocolVersion;
 
 			NodeInfo.StatUpdated(MoveTemp(Stat));
 		}
