UdpMessaging: Don't access errno outside of failure cases in SocketSender.

#rnx

diff --git a/Engine/Plugins/Messaging/UdpMessaging/Source/UdpMessaging/Private/Shared/SocketSender.h b/Engine/Plugins/Messaging/UdpMessaging/Source/UdpMessaging/Private/Shared/SocketSender.h
--- a/Engine/Plugins/Messaging/UdpMessaging/Source/UdpMessaging/Private/Shared/SocketSender.h
+++ b/Engine/Plugins/Messaging/UdpMessaging/Source/UdpMessaging/Private/Shared/SocketSender.h
@@ -193,10 +193,11 @@ protected:
 			const bool bSendResult =
 				Socket->SendTo(Packet->Data->GetData(), PacketNumBytes, SentNumBytes, *Recipient.ToInternetAddr());
 
-			const ESocketErrors SocketError = SocketSubsystem->GetLastErrorCode();
+			ESocketErrors SocketError = SE_NO_ERROR;
 
 			if (!bSendResult)
 			{
+				SocketError = SocketSubsystem->GetLastErrorCode();
 				if (SocketError == SE_EWOULDBLOCK || SocketError == SE_ENOBUFS || SocketError == SE_EINTR)
 				{
 					// Leave the packet in the queue to try again.
