Add better handling of delivery of a reassembled message in UdpMessaging. We hand off the message to the deserializer in a pending delivery state to better track the lifetime of the message.  Delivery should only be marked delivered after we have confirmed the deserialization task has completed.  The tracking flags have been marked atomic to better manage state changes across threads.

#rnx

diff --git a/Engine/Plugins/Messaging/UdpMessaging/Source/UdpMessaging/Private/Transport/UdpMessageProcessor.cpp b/Engine/Plugins/Messaging/UdpMessaging/Source/UdpMessaging/Private/Transport/UdpMessageProcessor.cpp
--- a/Engine/Plugins/Messaging/UdpMessaging/Source/UdpMessaging/Private/Transport/UdpMessageProcessor.cpp
+++ b/Engine/Plugins/Messaging/UdpMessaging/Source/UdpMessaging/Private/Transport/UdpMessageProcessor.cpp
@@ -780,7 +780,7 @@ void FUdpMessageProcessor::ProcessDataSegment(FInboundSegment& Segment, FNodeInf
 	}
 
 	// Deliver or re-sequence message
-	if (!ReassembledMessage->IsComplete() || ReassembledMessage->IsDelivered())
+	if (!ReassembledMessage->IsComplete() || ReassembledMessage->IsPendingDelivery())
 	{
 		return;
 	}
@@ -793,7 +793,7 @@ void FUdpMessageProcessor::ProcessDataSegment(FInboundSegment& Segment, FNodeInf
 		*NodeInfo.Endpoint.ToString());
 
 	AcknowledgeReceipt(DataChunk.MessageId, NodeInfo);
-	DeliverMessage(ReassembledMessage, NodeInfo);
+	TryDeliverMessage(ReassembledMessage, NodeInfo);
 }
 
 
@@ -943,7 +943,7 @@ void FUdpMessageProcessor::LookupAndCacheMessageType(TSharedPtr<FUdpReassembledM
 	}
 }
 
-void FUdpMessageProcessor::DeliverMessage(const TSharedPtr<FUdpReassembledMessage, ESPMode::ThreadSafe>& ReassembledMessage, FNodeInfo& NodeInfo)
+void FUdpMessageProcessor::TryDeliverMessage(const TSharedPtr<FUdpReassembledMessage, ESPMode::ThreadSafe>& ReassembledMessage, FNodeInfo& NodeInfo)
 {
 	// Do not deliver message while saving or garbage collecting since those deliveries will fail anyway...
 	if (ReassembledMessage->GetMessageTypeInfo() == nullptr && (UE::IsSavingPackage() || IsGarbageCollecting()))
@@ -957,8 +957,8 @@ void FUdpMessageProcessor::DeliverMessage(const TSharedPtr<FUdpReassembledMessag
 		MessageReassembledDelegate.ExecuteIfBound(ReassembledMessage, nullptr, NodeInfo.NodeId);
 	}
 
-	// Mark the message delivered but do not remove it from the list yet, this is to prevent the double delivery of reliable message
-	ReassembledMessage->MarkDelivered();
+	// Mark the message pending delivery but do not remove it from the list yet, this is to prevent the double delivery of reliable message.
+	ReassembledMessage->MarkPendingDelivery();
 }
 
 
@@ -1017,9 +1017,9 @@ bool FUdpMessageProcessor::MoreToSend()
 	return false;
 }
 
-void FUdpMessageProcessor::HandleSocketError(const FNodeInfo& NodeInfo) const
+void FUdpMessageProcessor::HandleSocketError(const TCHAR* ErrorReference, const FNodeInfo& NodeInfo) const
 {
-	UE_LOG(LogUdpMessaging, Error,  TEXT("Socket error detected when communicating with %s. Banning communication to that endpoint."), *NodeInfo.Endpoint.ToString());
+	UE_LOG(LogUdpMessaging, Error, TEXT("%s, Socket error detected when communicating with %s, Banning communication to that endpoint."), ErrorReference, *NodeInfo.Endpoint.ToString());
 	ErrorSendingToEndpointDelegate.Execute(NodeInfo.NodeId, NodeInfo.Endpoint);
 }
 
@@ -1132,7 +1132,7 @@ void FUdpMessageProcessor::UpdateKnownNodes()
 		if (!bSuccess || UE::Private::MessageProcessor::ShouldErrorOnConnection(KnownNodePair.Value.Endpoint))
 		{
 			bSuccess = false; // To ensure we trigger the delegate on a forced error.
-			HandleSocketError(KnownNodePair.Value);
+			HandleSocketError(TEXT("UpdateSegmenters"), KnownNodePair.Value);
 			break;
 		}
 
@@ -1140,7 +1140,7 @@ void FUdpMessageProcessor::UpdateKnownNodes()
 		// if there is a socket error, continuing is useless
 		if (!bSuccess)
 		{
-			HandleSocketError(KnownNodePair.Value);
+			HandleSocketError(TEXT("UpdateReassemblers"), KnownNodePair.Value);
 			break;
 		}
 
@@ -1461,6 +1461,13 @@ bool FUdpMessageProcessor::UpdateReassemblers(FNodeInfo& NodeInfo)
 	{
 		TSharedPtr<FUdpReassembledMessage, ESPMode::ThreadSafe>& ReassembledMessage = It.Value();
 
+		if (ReassembledMessage->IsDelivered())
+		{
+			// If this is already delivered no need to do anything else. 
+			It.RemoveCurrent();
+			continue;
+		}
+
 		int BytesSent = 0;
 		while (ReassembledMessage->HasPendingAcknowledgements() && BytesSent < MaxNodeByteSend)
 		{
@@ -1492,9 +1499,9 @@ bool FUdpMessageProcessor::UpdateReassemblers(FNodeInfo& NodeInfo)
 		LookupAndCacheMessageType(ReassembledMessage);
 
 		// Try to deliver completed message that couldn't be delivered the first time around
-		if (ReassembledMessage->IsComplete() && !ReassembledMessage->IsDelivered())
+		if (ReassembledMessage->IsComplete() && !ReassembledMessage->IsPendingDelivery())
 		{
-			DeliverMessage(ReassembledMessage, NodeInfo);
+			TryDeliverMessage(ReassembledMessage, NodeInfo);
 		}
 
 		// Remove stale reassembled message if they aren't reliable or are marked delivered
diff --git a/Engine/Plugins/Messaging/UdpMessaging/Source/UdpMessaging/Private/Transport/UdpMessageProcessor.h b/Engine/Plugins/Messaging/UdpMessaging/Source/UdpMessaging/Private/Transport/UdpMessageProcessor.h
--- a/Engine/Plugins/Messaging/UdpMessaging/Source/UdpMessaging/Private/Transport/UdpMessageProcessor.h
+++ b/Engine/Plugins/Messaging/UdpMessaging/Source/UdpMessaging/Private/Transport/UdpMessageProcessor.h
@@ -801,12 +801,12 @@ protected:
 	void ProcessUnknownSegment(FInboundSegment& Segment, FNodeInfo& NodeInfo, uint8 SegmentType);
 
 	/**
-	 * Deliver a reassembled message
+	 * Try to deliver the message to the transport layer.
 	 *
 	 * @param ReassembledMessage The message to deliver.
 	 * @param NodeInfo Details of the node that sent the message.
 	 */
-	void DeliverMessage(const TSharedPtr<FUdpReassembledMessage, ESPMode::ThreadSafe>& ReassembledMessage, FNodeInfo& NodeInfo);
+	void TryDeliverMessage(const TSharedPtr<FUdpReassembledMessage, ESPMode::ThreadSafe>& ReassembledMessage, FNodeInfo& NodeInfo);
 
 	/**
 	 * Lookup the reassembled message type and try to determine the source struct type.
@@ -919,7 +919,7 @@ private:
 	void OnPluginLoadingPhaseComplete(ELoadingPhase::Type LoadingPhase, bool bPhaseSuccessful);
 
 	/** Handles a communication error to a particular endpoint.*/
-	void HandleSocketError(const FNodeInfo& NodeInfo) const;
+	void HandleSocketError(const TCHAR* ErrorReference, const FNodeInfo& NodeInfo) const;
 
 	/** Checks all known nodes to see if any segmenters have NeedSending set to true. */
 	bool MoreToSend();
diff --git a/Engine/Plugins/Messaging/UdpMessaging/Source/UdpMessaging/Private/Transport/UdpMessageTransport.cpp b/Engine/Plugins/Messaging/UdpMessaging/Source/UdpMessaging/Private/Transport/UdpMessageTransport.cpp
--- a/Engine/Plugins/Messaging/UdpMessaging/Source/UdpMessaging/Private/Transport/UdpMessageTransport.cpp
+++ b/Engine/Plugins/Messaging/UdpMessaging/Source/UdpMessaging/Private/Transport/UdpMessageTransport.cpp
@@ -474,8 +474,12 @@ void FUdpMessageTransport::HandleProcessorMessageReassembled(FUdpReassembledMess
 		}
 		else
 		{
-			UE_LOG(LogUdpMessaging, Verbose, TEXT("Failed to deserialize message from %s"), *NodeId.ToString());
+			TWeakObjectPtr<UScriptStruct> TypeInfo = ReassembledMessagePtr->GetMessageTypeInfo();
+			FString TypeInfoStr = TypeInfo.IsValid() ? TypeInfo->GetName() : TEXT("UKNOWN TYPE");
+			UE_LOG(LogUdpMessaging, Verbose, TEXT("Failed to deserialize message of type %s from %s"), *TypeInfoStr, *NodeId.ToString());
 		}
+		// We always mark the message as delivered even if we can't deserialize it. There may be some messages that can't be deserialized because of missing type information.
+		ReassembledMessagePtr->MarkDelivered();
 	};
 
 	if (DeserializationTasks)
diff --git a/Engine/Plugins/Messaging/UdpMessaging/Source/UdpMessaging/Private/Transport/UdpReassembledMessage.h b/Engine/Plugins/Messaging/UdpMessaging/Source/UdpMessaging/Private/Transport/UdpReassembledMessage.h
--- a/Engine/Plugins/Messaging/UdpMessaging/Source/UdpMessaging/Private/Transport/UdpReassembledMessage.h
+++ b/Engine/Plugins/Messaging/UdpMessaging/Source/UdpMessaging/Private/Transport/UdpReassembledMessage.h
@@ -16,7 +16,13 @@ enum class EMessageFlags : uint32;
 class FUdpReassembledMessage
 {
 public:
-
+	enum class EDeliveryStatus : uint8
+	{
+		NotDelivered = 0,
+		PendingDelivery,
+		Delivered
+	};
+		
 	/** Default constructor. */
 	FUdpReassembledMessage() = default;
 
@@ -35,7 +41,7 @@ public:
 		, MessageFlags(InFlags)
 		, PendingSegmentsCount(SegmentCount)
 		, ReceivedBytes(0)
-		, bIsDelivered(false)
+		, DeliveryStatus(EDeliveryStatus::NotDelivered)
 		, RetransmitRequestsCount(0)
 		, Sender(InSender)
 		, Sequence(InSequence)
@@ -227,17 +233,32 @@ public:
 	 */
 	bool IsDelivered() const
 	{
-		return bIsDelivered;
+		return DeliveryStatus == EDeliveryStatus::Delivered;
 	}
 
+	/**
+	 * Get if the message is pending delivery through the transport layer.
+	 */
+	bool IsPendingDelivery() const
+	{
+		return DeliveryStatus == EDeliveryStatus::PendingDelivery;
+	}
+		
 	/**
 	 * Mark the reasembled message as delivered.
 	 */
 	void MarkDelivered()
 	{
-		bIsDelivered = true;
+		DeliveryStatus = EDeliveryStatus::Delivered;
 	}
 
+	/**
+	 * Mark the reasembled message as pending delivery.
+	 */
+	void MarkPendingDelivery()
+	{
+		DeliveryStatus = EDeliveryStatus::PendingDelivery;
+	}
 
 	/**
 	 * Indicates if the first segment of the data is ready.
@@ -261,7 +282,7 @@ public:
 	 */
 	void Reassemble(int32 SegmentNumber, int32 SegmentOffset, const TArray<uint8>& SegmentData, const FDateTime& CurrentTime)
 	{
-		if (IsMalformed() || !IsInitialized() || SegmentData.Num() == 0)
+		if (IsPendingDelivery() || IsMalformed() || !IsInitialized() || SegmentData.Num() == 0)
 		{
 			return;
 		}
@@ -360,8 +381,8 @@ private:
 	/** Holds the number of bytes received so far. */
 	int32 ReceivedBytes;
 
-	/** Holds if the reassembled message has been delivered */
-	bool bIsDelivered;
+	/** Holds if the reassembled message delivery status. This may get set from another thread which is why it is atomic. */
+	std::atomic<EDeliveryStatus> DeliveryStatus;
 
 	/** Holds the number of retransmit requests that were sent since the last segment was received. */
 	int32 RetransmitRequestsCount;
