There was a race condition in the transition of our message from pending to delivery. This could cause us to potentially double deliver a message.  This change ensures that we check both the pending and delivered state before we try to deliver a message again.

diff --git a/Engine/Plugins/Messaging/UdpMessaging/Source/UdpMessaging/Private/Transport/UdpMessageProcessor.cpp b/Engine/Plugins/Messaging/UdpMessaging/Source/UdpMessaging/Private/Transport/UdpMessageProcessor.cpp
--- a/Engine/Plugins/Messaging/UdpMessaging/Source/UdpMessaging/Private/Transport/UdpMessageProcessor.cpp
+++ b/Engine/Plugins/Messaging/UdpMessaging/Source/UdpMessaging/Private/Transport/UdpMessageProcessor.cpp
@@ -780,7 +780,7 @@ void FUdpMessageProcessor::ProcessDataSegment(FInboundSegment& Segment, FNodeInf
 	}
 
 	// Deliver or re-sequence message
-	if (!ReassembledMessage->IsComplete() || ReassembledMessage->IsPendingDelivery())
+	if (!ReassembledMessage->IsComplete() || ReassembledMessage->IsPendingDelivery() || ReassembledMessage->IsDelivered())
 	{
 		return;
 	}
@@ -1499,16 +1499,16 @@ bool FUdpMessageProcessor::UpdateReassemblers(FNodeInfo& NodeInfo)
 		LookupAndCacheMessageType(ReassembledMessage);
 
 		// Try to deliver completed message that couldn't be delivered the first time around
-		if (ReassembledMessage->IsComplete() && !ReassembledMessage->IsPendingDelivery())
+		if (ReassembledMessage->IsComplete() && !(ReassembledMessage->IsPendingDelivery() || ReassembledMessage->IsDelivered()))
 		{
 			TryDeliverMessage(ReassembledMessage, NodeInfo);
 		}
 
 		// Remove stale reassembled message if they aren't reliable or are marked delivered
 		if (ReassembledMessage->GetLastSegmentTime() + StaleReassemblyInterval <= CurrentTime &&
-			(!EnumHasAnyFlags(ReassembledMessage->GetFlags(), EMessageFlags::Reliable) || ReassembledMessage->IsDelivered()))
+			(!EnumHasAnyFlags(ReassembledMessage->GetFlags(), EMessageFlags::Reliable)))
 		{
-			if (!ReassembledMessage->IsDelivered())
+			if (!ReassembledMessage->IsDelivered() && !ReassembledMessage->IsPendingDelivery())
 			{
 				const int ReceivedSegments = ReassembledMessage->GetTotalSegmentsCount() - ReassembledMessage->GetPendingSegmentsCount();
 				UE_LOG(LogUdpMessaging, Warning, TEXT("FUdpMessageProcessor::UpdateReassemblers Discarding %d/%d of stale message segments from %s"),
@@ -1516,7 +1516,11 @@ bool FUdpMessageProcessor::UpdateReassemblers(FNodeInfo& NodeInfo)
 					ReassembledMessage->GetTotalSegmentsCount(),
 					*ReassembledMessage->Describe());
 			}
-			It.RemoveCurrent();
+
+			if (!ReassembledMessage->IsPendingDelivery())
+			{
+				It.RemoveCurrent();
+			}
 		}
 	}
 	return true;
