UdpMessaging: Fix race condition in deserializer struct lookup.

diff --git a/Engine/Plugins/Messaging/UdpMessaging/Source/UdpMessaging/Private/Transport/UdpDeserializedMessage.cpp b/Engine/Plugins/Messaging/UdpMessaging/Source/UdpMessaging/Private/Transport/UdpDeserializedMessage.cpp
--- a/Engine/Plugins/Messaging/UdpMessaging/Source/UdpMessaging/Private/Transport/UdpDeserializedMessage.cpp
+++ b/Engine/Plugins/Messaging/UdpMessaging/Source/UdpMessaging/Private/Transport/UdpDeserializedMessage.cpp
@@ -156,15 +156,17 @@ TWeakObjectPtr<UScriptStruct> DeserializeSafeLookupForFindObject_InThread(FStrin
 	const bool bRunningInTaskThread = (CurrentSettings->MaxConcurrentDeserializationTasks > 1);
 	if (bRunningInTaskThread)
 	{
-		// Since we are running in a task thread. We can just spin wait to get the type info.
-		while (UE::IsSavingPackage() || IsGarbageCollecting())
-		{
-			FPlatformProcess::SleepNoStats(0.01f);			
-		}
+		FGCScopeGuard Guard;
+		return FindFirstObjectSafe<UScriptStruct>(TypeName, EFindFirstObjectOptions::EnsureIfAmbiguous);
 	}
-	if (!UE::IsSavingPackage() && !IsGarbageCollecting())
+	else
 	{
-		return FindFirstObjectSafe<UScriptStruct>(TypeName, EFindFirstObjectOptions::EnsureIfAmbiguous);
+		// If we're not in a task, we want to avoid potentially stalling the processor thread.
+		FGCScopeTryGuard MaybeGuard;
+		if (MaybeGuard.LockSucceeded())
+		{
+			return FindFirstObjectSafe<UScriptStruct>(TypeName, EFindFirstObjectOptions::EnsureIfAmbiguous);
+		}
 	}
 
 	return nullptr;
