UdpMessaging: Enable UDP segmentation offload.

FUdpMessageProcessor::SendNextSegmentForMessageId now supports coalescing multiple segments into a single "superpacket," and using a (potentially hardware-accelerated) fast path to divide it into multiple datagrams on the wire with reduced overhead.

- Bump UDP_MESSAGING_TRANSPORT_PROTOCOL_VERSION to 18.
- Add support for UDP_SEND_MSG_SIZE (Windows) to FSocketBSD.
- Replace UDP_MESSAGING_SEGMENT_SIZE define with configurable UUdpMessagingSettings::MaxPacketSize, and attempt to handle it in a backward-compatible manner.
- Maximum USO generated datagram count is controllable via the `MessageBus.UDP.MaxOffloadSegments` cvar (default 16, max 64).
- Optimize FUdpMessageSegmenter operations on segment bit arrays; we now cache and skip the contiguous ack'd range at the start.
- Move LookupAndCacheMessageType off the hot path; try once on delivery of the first segment, and retry once each tick after until success.
- Fix bugs where AcknowledgeSegments and Mesh segments could be generated that exceeded the max segment size.
- Fix incorrect slice count in case where multiple Mesh segments are generated.
- Fix uint32 overflow in FNodeInfo::Remap Szudzik pairing calculation.
- Deprecate UUdpMessagingSettings::WorkQueueSize, which is no longer used.

diff --git a/Engine/Plugins/Messaging/UdpMessaging/Source/UdpMessaging/Private/Shared/SocketSender.cpp b/Engine/Plugins/Messaging/UdpMessaging/Source/UdpMessaging/Private/Shared/SocketSender.cpp
new file mode 100644
--- /dev/null
+++ b/Engine/Plugins/Messaging/UdpMessaging/Source/UdpMessaging/Private/Shared/SocketSender.cpp
@@ -0,0 +1,215 @@
+// Copyright Epic Games, Inc. All Rights Reserved.
+
+#include "SocketSender.h"
+
+#include "HAL/Event.h"
+#include "HAL/PlatformProcess.h"
+#include "HAL/RunnableThread.h"
+#include "Logging/StructuredLog.h"
+#include "Sockets.h"
+#include "SocketSubsystem.h"
+#include "UdpMessagingPrivate.h"
+
+
+namespace UE::UdpMessaging {
+
+
+FSocketSender::FSocketSender(FSocket* InSocket, const TCHAR* InDescription, const FOptions& InOptions /* = FOptions() */)
+	: Socket(InSocket)
+	, SocketSubsystem(ISocketSubsystem::Get(PLATFORM_SOCKETSUBSYSTEM))
+	, bStopping(false)
+	, WaitTime(InOptions.WaitTime)
+	, Description(InDescription)
+{
+	check(Socket != nullptr);
+	check(Socket->GetSocketType() == SOCKTYPE_Datagram);
+
+	int32 NewSize = 0;
+	Socket->SetSendBufferSize(InOptions.SendBufferSize, NewSize);
+
+	WorkEvent = FPlatformProcess::GetSynchEventFromPool();
+	Thread = FRunnableThread::Create(this, InDescription, 128 * 1024, TPri_AboveNormal,
+		FPlatformAffinity::GetPoolThreadMask());
+
+	if (InOptions.SegmentOffloadSize != 0)
+	{
+		TryEnableSegmentationOffload(InOptions.SegmentOffloadSize);
+	}
+	else
+	{
+		SegmentOffloadSize = 0;
+	}
+}
+
+
+FSocketSender::~FSocketSender()
+{
+	if (Thread != nullptr)
+	{
+		Thread->Kill(true);
+		delete Thread;
+	}
+
+	FPlatformProcess::ReturnSynchEventToPool(WorkEvent);
+	WorkEvent = nullptr;
+}
+
+
+bool FSocketSender::Send(
+	const TSharedRef<TArray<uint8>>& InData,
+	const FIPv4Endpoint& InRecipient,
+	bool bAllowSegmentation /* = false */
+)
+{
+	if (!bStopping)
+	{
+		if (!bAllowSegmentation && SegmentOffloadSize > 0)
+		{
+			// If this hits, your call to Send() is going to generate multiple packets on the wire.
+			// If this was intentional, pass bAllowSegmentation = true.
+			checkf(InData->Num() <= SegmentOffloadSize,
+				TEXT("Send data exceeds USO offload size, but bAllowSegmentation was false"));
+		}
+
+		SendQueue.Enqueue(InData, InRecipient);
+		WorkEvent->Trigger();
+		return true;
+	}
+
+	return false;
+}
+
+uint32 FSocketSender::Run()
+{
+	while (!bStopping)
+	{
+		const EUpdateResult Result = Update(WaitTime);
+		switch (Result)
+		{
+			case EUpdateResult::Done:
+				WorkEvent->Wait(WaitTime);
+				continue;
+			case EUpdateResult::Retry:
+				continue;
+			case EUpdateResult::Fatal:
+				bStopping = true;
+				return 0;
+		}
+	}
+
+	return 0;
+}
+
+
+void FSocketSender::Stop()
+{
+	bStopping = true;
+	WorkEvent->Trigger();
+}
+
+
+FSocketSender::EUpdateResult FSocketSender::Update(const FTimespan& InSocketWaitTime)
+{
+	while (!SendQueue.IsEmpty())
+	{
+		if (!Socket->Wait(ESocketWaitConditions::WaitForWrite, InSocketWaitTime))
+		{
+			return EUpdateResult::Retry;
+		}
+
+		const FPacket* Packet = SendQueue.Peek();
+		const FIPv4Endpoint Recipient = Packet->Recipient;
+		const int32 PacketNumBytes = Packet->Data->Num();
+		int32 SentNumBytes = 0;
+		const bool bSendResult =
+			Socket->SendTo(Packet->Data->GetData(), PacketNumBytes, SentNumBytes, *Recipient.ToInternetAddr());
+
+		ESocketErrors SocketError = SE_NO_ERROR;
+
+		if (!bSendResult)
+		{
+			SocketError = SocketSubsystem->GetLastErrorCode();
+			if (SocketError == SE_EWOULDBLOCK || SocketError == SE_ENOBUFS || SocketError == SE_EINTR)
+			{
+				// Leave the packet in the queue to try again.
+				return EUpdateResult::Retry;
+			}
+		}
+
+		// Only pop after we've checked for "retry-able" error codes.
+		Packet = nullptr;
+		SendQueue.Dequeue();
+
+		if (!bSendResult)
+		{
+			bool bFatalError;
+			switch (SocketError)
+			{
+				case SE_ENETDOWN:
+				case SE_ENETUNREACH:
+				case SE_EHOSTDOWN:
+				case SE_EHOSTUNREACH:
+				case SE_EADDRNOTAVAIL:
+					bFatalError = false;
+					break;
+				default:
+					bFatalError = true;
+					break;
+			}
+
+			if (!bFatalError)
+			{
+				// Non-fatal; on to the next packet in the queue.
+				UE_LOGFMT(LogUdpMessaging, Verbose,
+					"Sender {Description}: SendTo failed (destination: {Recipient}) ({ErrorStr})",
+					Description, Recipient.ToString(), SocketSubsystem->GetSocketError(SocketError));
+				continue;
+			}
+			else
+			{
+				UE_LOGFMT(LogUdpMessaging, Error,
+					"Sender {Description}: SendTo failed (destination: {Recipient}) ({ErrorStr})",
+					Description, Recipient.ToString(), SocketSubsystem->GetSocketError(SocketError));
+			}
+		}
+
+		if (SentNumBytes != PacketNumBytes)
+		{
+			if (SentNumBytes >= 0)
+			{
+				// FIXME?: In the absence of another socket error, could be a retry? Is this even possible for UDP?
+				UE_LOGFMT(LogUdpMessaging, Error,
+					"Sender {Description}: Incomplete send (destination: {Recipient}) ({Sent}/{Total} bytes)",
+					Description, Recipient.ToString(), SentNumBytes, PacketNumBytes);
+			}
+			return EUpdateResult::Fatal;
+		}
+	}
+
+	return EUpdateResult::Done;
+}
+
+
+void FSocketSender::TryEnableSegmentationOffload(uint16 InSegmentOffloadSize)
+{
+	UE_LOGFMT(LogUdpMessaging, Log, "Sender {Description}: Attempting to enable segmentation offload", Description);
+
+	if (Socket->SetSendSegmentationOffloadSize(InSegmentOffloadSize))
+	{
+		UE_LOGFMT(LogUdpMessaging, Log, "Sender {Description}: Segmentation offload enabled ({SegmentSize} bytes)",
+			Description, InSegmentOffloadSize);
+
+		SegmentOffloadSize = InSegmentOffloadSize;
+	}
+	else
+	{
+		const ESocketErrors SocketError = SocketSubsystem->GetLastErrorCode();
+		UE_LOGFMT(LogUdpMessaging, Warning, "Sender {Description}: Failed to enable segmentation offload ({Error})",
+			Description, SocketSubsystem->GetSocketError());
+
+		SegmentOffloadSize = 0;
+	}
+}
+
+
+} // namespace UE::UdpMessaging
diff --git a/Engine/Plugins/Messaging/UdpMessaging/Source/UdpMessaging/Private/Shared/SocketSender.h b/Engine/Plugins/Messaging/UdpMessaging/Source/UdpMessaging/Private/Shared/SocketSender.h
--- a/Engine/Plugins/Messaging/UdpMessaging/Source/UdpMessaging/Private/Shared/SocketSender.h
+++ b/Engine/Plugins/Messaging/UdpMessaging/Source/UdpMessaging/Private/Shared/SocketSender.h
@@ -4,17 +4,16 @@
 
 #include <atomic>
 #include "Containers/MpscQueue.h"
-#include "HAL/Event.h"
-#include "HAL/PlatformProcess.h"
 #include "HAL/Runnable.h"
-#include "HAL/RunnableThread.h"
-#include "Logging/StructuredLog.h"
+#include "Interfaces/IPv4/IPv4Endpoint.h"
 #include "Misc/SingleThreadRunnable.h"
-#include "Sockets.h"
-#include "SocketSubsystem.h"
-#include "UdpMessagingPrivate.h"
+#include "Templates/SharedPointer.h"
 
-#include "Interfaces/IPv4/IPv4Endpoint.h"
+
+class FEvent;
+class FRunnableThread;
+class FSocket;
+class ISocketSubsystem;
 
 
 namespace UE::UdpMessaging {
@@ -58,14 +57,19 @@ class FSocketSender
 	};
 
 public:
-
 	struct FOptions
 	{
 		int32 SendBufferSize = 512 * 1024;
 		FTimespan WaitTime = FTimespan::FromMilliseconds(100);
 
+		/** Datagram payload size for UDP segmentation offload; 0 disables USO. Ignored if unsupported. */
+		uint16 SegmentOffloadSize = 0;
+
 		FOptions()
 		{
+			// Without this (or with `= default`), clang gives the following:
+			// `error: default member initializer for 'SendBufferSize' needed within definition
+			// of enclosing class 'FSocketSender' outside of member functions` (et cetera)
 		}
 	};
 
@@ -74,62 +78,31 @@ public:
 	 *
 	 * @param InSocket The UDP socket to use for sending data.
 	 * @param InDescription The description text (for debugging).
+	 * @param InOptions Other configuration.
 	 */
-	FSocketSender(FSocket* InSocket, const TCHAR* InDescription, const FOptions& InOptions = FOptions())
-		: Socket(InSocket)
-		, SocketSubsystem(ISocketSubsystem::Get(PLATFORM_SOCKETSUBSYSTEM))
-		, bStopping(false)
-		, WaitTime(InOptions.WaitTime)
-		, Description(InDescription)
-	{
-		check(Socket != nullptr);
-		check(Socket->GetSocketType() == SOCKTYPE_Datagram);
-
-		int32 NewSize = 0;
-		Socket->SetSendBufferSize(InOptions.SendBufferSize, NewSize);
-
-		WorkEvent = FPlatformProcess::GetSynchEventFromPool();
-		Thread = FRunnableThread::Create(this, InDescription, 128 * 1024, TPri_AboveNormal, FPlatformAffinity::GetPoolThreadMask());
-	}
+	FSocketSender(FSocket* InSocket, const TCHAR* InDescription, const FOptions& InOptions = FOptions());
 
 	/** Virtual destructor. */
-	virtual ~FSocketSender()
-	{
-		if (Thread != nullptr)
-		{
-			Thread->Kill(true);
-			delete Thread;
-		}
-
-		FPlatformProcess::ReturnSynchEventToPool(WorkEvent);
-		WorkEvent = nullptr;
-	}
-
-public:
+	virtual ~FSocketSender();
 
 	/**
 	 * Sends data to the specified recipient.
 	 *
-	 * @param Data The data to send.
-	 * @param Recipient The recipient.
+	 * @param InData The data to send.
+	 * @param InRecipient The recipient.
+	 * @param bAllowSegmentation Pass true if you are intentionally sending more than `GetSegmentOffloadSize()` bytes.
 	 * @return true if the data will be sent, false otherwise.
 	 */
-	bool Send(const TSharedRef<TArray<uint8>, ESPMode::ThreadSafe>& Data, const FIPv4Endpoint& Recipient)
-	{
-		if (!bStopping)
-		{
-			SendQueue.Enqueue(Data, Recipient);
-			WorkEvent->Trigger();
-			return true;
-		}
+	bool Send(
+		const TSharedRef<TArray<uint8>>& InData,
+		const FIPv4Endpoint& InRecipient,
+		bool bAllowSegmentation = false
+	);
 
-		return false;
-	}
-
-public:
+	/** If send segmentation offload is enabled, returns the max datagram payload. Otherwise 0. */
+	uint16 GetSegmentOffloadSize() const { return SegmentOffloadSize; }
 
 	//~ FRunnable interface
-
 	virtual FSingleThreadRunnable* GetSingleThreadInterface() override
 	{
 		return this;
@@ -140,135 +113,27 @@ public:
 		return true;
 	}
 
-	virtual uint32 Run() override
-	{
-		while (!bStopping)
-		{
-			const EUpdateResult Result = Update(WaitTime);
-			switch (Result)
-			{
-				case EUpdateResult::Done:
-					WorkEvent->Wait(WaitTime);
-					continue;
-				case EUpdateResult::Retry:
-					continue;
-				case EUpdateResult::Fatal:
-					bStopping = true;
-					return 0;
-			}
-		}
-
-		return 0;
-	}
+	virtual uint32 Run() override;
+	virtual void Stop() override;
+	virtual void Exit() override { }
 
-	virtual void Stop() override
+	//~ FSingleThreadRunnable interface
+	virtual void Tick() override
 	{
-		bStopping = true;
-		WorkEvent->Trigger();
+		Update(FTimespan::Zero());
 	}
 
-	virtual void Exit() override { }
-
 protected:
-
 	/**
 	 * Update this socket sender.
 	 *
-	 * @param Time to wait for the socket.
-	 * @return true on success, false otherwise.
+	 * @param InSocketWaitTime Time to wait for socket write conditions.
 	 */
-	EUpdateResult Update(const FTimespan& SocketWaitTime)
-	{
-		while (!SendQueue.IsEmpty())
-		{
-			if (!Socket->Wait(ESocketWaitConditions::WaitForWrite, SocketWaitTime))
-			{
-				return EUpdateResult::Retry;
-			}
-
-			const FPacket* Packet = SendQueue.Peek();
-			const FIPv4Endpoint Recipient = Packet->Recipient;
-			const int32 PacketNumBytes = Packet->Data->Num();
-			int32 SentNumBytes = 0;
-			const bool bSendResult =
-				Socket->SendTo(Packet->Data->GetData(), PacketNumBytes, SentNumBytes, *Recipient.ToInternetAddr());
-
-			ESocketErrors SocketError = SE_NO_ERROR;
-
-			if (!bSendResult)
-			{
-				SocketError = SocketSubsystem->GetLastErrorCode();
-				if (SocketError == SE_EWOULDBLOCK || SocketError == SE_ENOBUFS || SocketError == SE_EINTR)
-				{
-					// Leave the packet in the queue to try again.
-					return EUpdateResult::Retry;
-				}
-			}
-
-			// Only pop after we've checked for "retry-able" error codes.
-			Packet = nullptr;
-			SendQueue.Dequeue();
-
-			if (!bSendResult)
-			{
-				bool bFatalError;
-				switch (SocketError)
-				{
-					case SE_ENETDOWN:
-					case SE_ENETUNREACH:
-					case SE_EHOSTDOWN:
-					case SE_EHOSTUNREACH:
-					case SE_EADDRNOTAVAIL:
-						bFatalError = false;
-						break;
-					default:
-						bFatalError = true;
-						break;
-				}
-
-				if (!bFatalError)
-				{
-					// Non-fatal; on to the next packet in the queue.
-					UE_LOGFMT(LogUdpMessaging, Verbose,
-						"Sender {Description}: SendTo failed (destination: {Recipient}) ({ErrorStr})",
-						Description, Recipient.ToString(), SocketSubsystem->GetSocketError(SocketError));
-					continue;
-				}
-				else
-				{
-					UE_LOGFMT(LogUdpMessaging, Error,
-						"Sender {Description}: SendTo failed (destination: {Recipient}) ({ErrorStr})",
-						Description, Recipient.ToString(), SocketSubsystem->GetSocketError(SocketError));
-				}
-			}
-
-			if (SentNumBytes != PacketNumBytes)
-			{
-				if (SentNumBytes >= 0)
-				{
-					// FIXME?: In the absence of another socket error, could be a retry? Is this even possible for UDP?
-					UE_LOGFMT(LogUdpMessaging, Error,
-						"Sender {Description}: Incomplete send (destination: {Recipient}) ({Sent}/{Total} bytes)",
-						Description, Recipient.ToString(), SentNumBytes, PacketNumBytes);
-				}
-				return EUpdateResult::Fatal;
-			}
-		}
+	EUpdateResult Update(const FTimespan& InSocketWaitTime);
 
-		return EUpdateResult::Done;
-	}
-
-protected:
-
-	//~ FSingleThreadRunnable interface
-
-	virtual void Tick() override
-	{
-		Update(FTimespan::Zero());
-	}
+	void TryEnableSegmentationOffload(uint16 InSegmentOffloadSize);
 
 private:
-
 	/** The send queue. */
 	TMpscQueue<FPacket> SendQueue;
 
@@ -292,6 +157,9 @@ private:
 
 	/** Description of this sender (for debugging). */
 	FString Description;
+
+	/** Datagram payload size for segmentation offload. 0 means disabled or unsupported. */
+	uint16 SegmentOffloadSize;
 };
 
 } // namespace UE::UdpMessaging
diff --git a/Engine/Plugins/Messaging/UdpMessaging/Source/UdpMessaging/Private/Shared/UdpMessageSegment.h b/Engine/Plugins/Messaging/UdpMessaging/Source/UdpMessaging/Private/Shared/UdpMessageSegment.h
--- a/Engine/Plugins/Messaging/UdpMessaging/Source/UdpMessaging/Private/Shared/UdpMessageSegment.h
+++ b/Engine/Plugins/Messaging/UdpMessaging/Source/UdpMessaging/Private/Shared/UdpMessageSegment.h
@@ -76,6 +76,13 @@ namespace FUdpMessageSegment
 		 */
 		EUdpMessageSegments SegmentType;
 
+		static constexpr uint32 SerializedBytes = (
+			sizeof(ProtocolVersion) +
+			sizeof(RecipientNodeId) +
+			sizeof(SenderNodeId) +
+			sizeof(SegmentType)
+		);
+
 	public:
 
 		/**
@@ -182,6 +189,11 @@ namespace FUdpMessageSegment
 		/** List of Acknowledged segments */
 		TArray<uint32> Segments;
 
+		static constexpr uint32 FixedSerializedBytes = (
+			sizeof(MessageId) +
+			sizeof(decltype(Segments)::SizeType)
+		);
+
 	public:
 		FAcknowledgeSegmentsChunk() = default;
 
@@ -297,6 +309,17 @@ namespace FUdpMessageSegment
 		/** Holds the segment data. */
 		TArray<uint8> Data;
 
+		static constexpr uint32 FixedSerializedBytes = (
+			sizeof(MessageId) +
+			sizeof(MessageSize) +
+			sizeof(MessageFlags) +
+			sizeof(SegmentNumber) +
+			sizeof(SegmentOffset) +
+			sizeof(TotalSegments) +
+			sizeof(Sequence) +
+			sizeof(decltype(Data)::SizeType)
+		);
+
 	public:
 		FDataChunk() = default;
 
diff --git a/Engine/Plugins/Messaging/UdpMessaging/Source/UdpMessaging/Private/Tests/UdpMessageSegmenterTest.cpp b/Engine/Plugins/Messaging/UdpMessaging/Source/UdpMessaging/Private/Tests/UdpMessageSegmenterTest.cpp
--- a/Engine/Plugins/Messaging/UdpMessaging/Source/UdpMessaging/Private/Tests/UdpMessageSegmenterTest.cpp
+++ b/Engine/Plugins/Messaging/UdpMessaging/Source/UdpMessaging/Private/Tests/UdpMessageSegmenterTest.cpp
@@ -12,20 +12,23 @@
 IMPLEMENT_SIMPLE_AUTOMATION_TEST(FUdpMessageSegmenterTest, "System.Core.Messaging.Transports.Udp.UdpMessageSegmenter", EAutomationTestFlags_ApplicationContextMask | EAutomationTestFlags::EngineFilter)
 
 
-void RunSegmentationTest(FAutomationTestBase& Test, uint32 MessageSize, uint16 SegmentSize, bool WithReliableMessages)
+void RunSegmentationTest(FAutomationTestBase& Test, uint32 MessageSize, uint16 SegmentPayloadSize, bool WithReliableMessages)
 {
-	check(MessageSize < 255u * SegmentSize);
+	check(MessageSize < (uint32)255 * SegmentPayloadSize);
 
-	uint16 NumSegments = (MessageSize + SegmentSize - 1) / SegmentSize;
-	Test.AddInfo(FString::Printf(TEXT("Segmenting message of size %i with %i segments of size %i (Reliable=%d)..."), MessageSize, NumSegments, SegmentSize, WithReliableMessages));
+	uint16 NumSegments = (MessageSize + SegmentPayloadSize - 1) / SegmentPayloadSize;
+	Test.AddInfo(FString::Printf(TEXT("Segmenting message of size %i with %i payloads of size %i (Reliable=%d)..."),
+		MessageSize, NumSegments, SegmentPayloadSize, WithReliableMessages));
 
 	// create a large message to segment
-	TSharedRef<FUdpSerializedMessage, ESPMode::ThreadSafe> Message = MakeShared<FUdpSerializedMessage, ESPMode::ThreadSafe>(EUdpMessageFormat::CborPlatformEndianness, UDP_MESSAGING_TRANSPORT_PROTOCOL_VERSION, WithReliableMessages ? EMessageFlags::Reliable : EMessageFlags::None);
+	TSharedRef<FUdpSerializedMessage> Message = MakeShared<FUdpSerializedMessage>(
+		EUdpMessageFormat::CborPlatformEndianness, UDP_MESSAGING_TRANSPORT_PROTOCOL_VERSION,
+		WithReliableMessages ? EMessageFlags::Reliable : EMessageFlags::None);
 
 	for (uint8 SegmentIndex = 0; SegmentIndex < NumSegments; ++SegmentIndex)
 	{
 		// write segment index into each byte of the current segment
-		for (uint16 Offset = 0; (Offset < SegmentSize) && (Message->TotalSize() < MessageSize); ++Offset)
+		for (uint16 Offset = 0; (Offset < SegmentPayloadSize) && (Message->TotalSize() < MessageSize); ++Offset)
 		{
 			*Message << SegmentIndex;
 		}
@@ -34,6 +37,7 @@ void RunSegmentationTest(FAutomationTestBase& Test, uint32 MessageSize, uint16 S
 	Message->UpdateState(EUdpSerializedMessageState::Complete);
 
 	// create and initialize segmenter
+	const uint16 SegmentSize = IntCastChecked<uint16>((uint32)SegmentPayloadSize + FUdpMessageSegmenter::GetDataSegmentHeaderSize());
 	FUdpMessageSegmenter Segmenter = FUdpMessageSegmenter(Message, SegmentSize);
 	Segmenter.Initialize();
 
@@ -102,12 +106,12 @@ void RunSegmentationTest(FAutomationTestBase& Test, uint32 MessageSize, uint16 S
 
 			// verify segment size
 			int32 ExpectedSegmentSize = (OutSegmentNumber == (NumSegments - 1))
-				? MessageSize % SegmentSize
-				: SegmentSize;
+				? MessageSize % SegmentPayloadSize
+				: SegmentPayloadSize;
 
 			if (ExpectedSegmentSize == 0)
 			{
-				ExpectedSegmentSize = SegmentSize;
+				ExpectedSegmentSize = SegmentPayloadSize;
 			}
 
 			if (OutData.Num() != ExpectedSegmentSize)
diff --git a/Engine/Plugins/Messaging/UdpMessaging/Source/UdpMessaging/Private/Transport/UdpDeserializedMessage.cpp b/Engine/Plugins/Messaging/UdpMessaging/Source/UdpMessaging/Private/Transport/UdpDeserializedMessage.cpp
--- a/Engine/Plugins/Messaging/UdpMessaging/Source/UdpMessaging/Private/Transport/UdpDeserializedMessage.cpp
+++ b/Engine/Plugins/Messaging/UdpMessaging/Source/UdpMessaging/Private/Transport/UdpDeserializedMessage.cpp
@@ -44,7 +44,7 @@ class FUdpDeserializedMessageDetails
 public:
 	static bool DeserializeV10(FUdpDeserializedMessage& DeserializedMessage, FMemoryReader& MessageReader);
 	static bool DeserializeV11_15(FUdpDeserializedMessage& DeserializedMessage, FMemoryReader& MessageReader, bool bIsLWCBackwardCompatibilityMode);
-	static bool DeserializeV16_17(FUdpDeserializedMessage& DeserializedMessage, FMemoryReader& MessageReader);
+	static bool DeserializeV16_18(FUdpDeserializedMessage& DeserializedMessage, FMemoryReader& MessageReader);
 	static bool Deserialize(FUdpDeserializedMessage& DeserializedMessage, const FUdpReassembledMessage& ReassembledMessage);
 
 	template <typename ReturnType>
@@ -395,7 +395,7 @@ bool FUdpDeserializedMessageDetails::DeserializeV11_15(FUdpDeserializedMessage&
 	}
 }
 
-bool FUdpDeserializedMessageDetails::DeserializeV16_17(FUdpDeserializedMessage& DeserializedMessage, FMemoryReader& MessageReader)
+bool FUdpDeserializedMessageDetails::DeserializeV16_18(FUdpDeserializedMessage& DeserializedMessage, FMemoryReader& MessageReader)
 {
 	// message type info
 	{
@@ -544,6 +544,8 @@ FNameOrFTopLevel FUdpDeserializedMessage::PeekMessageTypeInfoName(const FUdpReas
 	case 16:
 		// fallthrough
 	case 17:
+		// fallthrough
+	case 18:
 	{
 		FTopLevelAssetPath AssetPath = FUdpDeserializedMessageDetails::GetMessageType<FTopLevelAssetPath>(MessageReader);
 		Path.Set<FTopLevelAssetPath>(AssetPath);
@@ -611,8 +613,10 @@ bool FUdpDeserializedMessageDetails::Deserialize(FUdpDeserializedMessage& Deseri
 	case 16:
 		// fallthrough
 	case 17:
+		// fallthrough
+	case 18:
 	{
-		return DeserializeV16_17(DeserializedMessage, MessageReader);
+		return DeserializeV16_18(DeserializedMessage, MessageReader);
 	}
 
 	default:
@@ -621,4 +625,3 @@ bool FUdpDeserializedMessageDetails::Deserialize(FUdpDeserializedMessage& Deseri
 	}
 	return false;
 }
-
diff --git a/Engine/Plugins/Messaging/UdpMessaging/Source/UdpMessaging/Private/Transport/UdpMessageBeacon.h b/Engine/Plugins/Messaging/UdpMessaging/Source/UdpMessaging/Private/Transport/UdpMessageBeacon.h
--- a/Engine/Plugins/Messaging/UdpMessaging/Source/UdpMessaging/Private/Transport/UdpMessageBeacon.h
+++ b/Engine/Plugins/Messaging/UdpMessaging/Source/UdpMessaging/Private/Transport/UdpMessageBeacon.h
@@ -34,7 +34,6 @@ public:
 	 * @param InSocketSender The network socket sender used to send Hello segments.
 	 * @param InSocketId The network socket identifier (used to detect unicast endpoint).
 	 * @param InMulticastEndpoint The multicast group endpoint to transport messages to.
-	 * @param InStaticEndpoints The static nodes to broadcast to alongside the multicast.
 	 */
 	FUdpMessageBeacon(
 		UE::UdpMessaging::FSocketSender* InSocketSender,
diff --git a/Engine/Plugins/Messaging/UdpMessaging/Source/UdpMessaging/Private/Transport/UdpMessageProcessor.cpp b/Engine/Plugins/Messaging/UdpMessaging/Source/UdpMessaging/Private/Transport/UdpMessageProcessor.cpp
--- a/Engine/Plugins/Messaging/UdpMessaging/Source/UdpMessaging/Private/Transport/UdpMessageProcessor.cpp
+++ b/Engine/Plugins/Messaging/UdpMessaging/Source/UdpMessaging/Private/Transport/UdpMessageProcessor.cpp
@@ -4,6 +4,7 @@
 #include "Algo/AllOf.h"
 
 #include "CoreGlobals.h"
+#include "HAL/RunnableThread.h"
 #include "Misc/ConfigCacheIni.h"
 #include "Serialization/ArrayReader.h"
 #include "Serialization/ArrayWriter.h"
@@ -53,7 +54,13 @@ TAutoConsoleVariable<int32> CVarCheckForExpired(
 TAutoConsoleVariable<int32> CVarUdpMessagingMaxWindowSize(
 	TEXT("MessageBus.UDP.MaxWindowSize"),
 	2048,
-	TEXT("Maximum window size for sending to endpoints.\n"),
+	TEXT("Maximum window size for sending to endpoints."),
+	ECVF_Default);
+
+TAutoConsoleVariable<int32> CVarUdpMessagingMaxOffloadSegments(
+	TEXT("MessageBus.UDP.MaxOffloadSegments"),
+	16,
+	TEXT("Maximum number of datagrams to generate per send call using UDP segmentation offload. Maximum value is 64."),
 	ECVF_Default);
 
 
@@ -97,20 +104,16 @@ FOnInboundTransferDataUpdated& OnReassemblerUpdated()
 	return OnTransferUpdated;
 }
 
-uint16 GetMessageProcessorWorkQueueSize()
-{
-	return GetDefault<UUdpMessagingSettings>()->WorkQueueSize;
 }
 
-}
 /* FUdpMessageProcessor structors
  *****************************************************************************/
 FUdpMessageProcessor::FNodeInfo::FNodeInfo()
 	: LastSegmentReceivedTime(FDateTime::MinValue())
 	, NodeId()
 	, ProtocolVersion(UDP_MESSAGING_TRANSPORT_PROTOCOL_VERSION)
-	, ReliableWorkQueue(UE::Private::MessageProcessor::GetMessageProcessorWorkQueueSize())
-	, UnreliableWorkQueue(UE::Private::MessageProcessor::GetMessageProcessorWorkQueueSize())
+	, ReliableWorkQueue(1024)
+	, UnreliableWorkQueue(1024)
 {
 	ComputeWindowSize(0, 0);
 }
@@ -310,15 +313,17 @@ void FUdpMessageProcessor::SendSegmenterStatsToListeners(int32 MessageId, FGuid
 	{
 		return;
 	}
-	uint32 SegmentCount = Segmenter->GetSegmentCount();
+
+	const uint16 MaxSegmentSize = Segmenter->GetTotalSegmentSize();
+	const uint32 SegmentCount = Segmenter->GetSegmentCount();
 	SegmenterUpdatedDelegate.Broadcast(
 		{
 			NodeId,
 			MessageId,
 			// Convert segment data into bytes.
-			UDP_MESSAGING_SEGMENT_SIZE * SegmentCount,
-			UDP_MESSAGING_SEGMENT_SIZE * (SegmentCount - Segmenter->GetPendingSendSegmentsCount()),
-			UDP_MESSAGING_SEGMENT_SIZE * (Segmenter->GetAcknowledgedSegmentsCount()),
+			MaxSegmentSize * SegmentCount,
+			MaxSegmentSize * (SegmentCount - Segmenter->GetPendingSendSegmentsCount()),
+			MaxSegmentSize * (Segmenter->GetAcknowledgedSegmentsCount()),
 			EnumHasAnyFlags(Segmenter->GetMessageFlags(), EMessageFlags::Reliable),
 			EnumHasAnyFlags(Segmenter->GetMessageFlags(), EMessageFlags::Traced),
 			Segmenter->AreAcknowledgementsComplete()
@@ -348,13 +353,22 @@ bool FUdpMessageProcessor::Init()
 {
 	if (!bIsInitialized)
 	{
-		SocketSender = MakeUnique<UE::UdpMessaging::FSocketSender>(Socket, TEXT("FUdpMessageProcessor.Sender"));
+		const UUdpMessagingSettings* Settings = GetDefault<UUdpMessagingSettings>();
+
+		UE::UdpMessaging::FSocketSender::FOptions SenderOpts;
+		SenderOpts.SegmentOffloadSize =
+			Settings->bSegmentationOffload
+			? (Settings->GetMaxPacketSize() - UE::UdpMessaging::PacketHeaderBytes)
+			: 0;
+
+		SocketSender = MakeUnique<UE::UdpMessaging::FSocketSender>(Socket, TEXT("FUdpMessageProcessor.Sender"), SenderOpts);
 
 		Beacon = MakeUnique<FUdpMessageBeacon>(SocketSender.Get(), LocalNodeId, MulticastEndpoint);
 
-		// Current protocol version 17
+		// Current protocol version 18
 		SupportedProtocolVersions.Add(UDP_MESSAGING_TRANSPORT_PROTOCOL_VERSION);
-		// Support Protocol version 10, 11, 12, 13, 14, 15, 16
+		// Support Protocol version 10, 11, 12, 13, 14, 15, 16, 17
+		SupportedProtocolVersions.Add(17);
 		SupportedProtocolVersions.Add(16);
 		SupportedProtocolVersions.Add(15);
 		SupportedProtocolVersions.Add(14);
@@ -380,6 +394,14 @@ uint32 FUdpMessageProcessor::Run()
 			CurrentTime = FDateTime::UtcNow();
 			DeltaTime = CurrentTime - LastTime;
 
+			for (TSet<FUdpReassembledMessageSharedPtr>::TIterator It = PendingTypeLookup.CreateIterator(); It; ++It)
+			{
+				if (LookupAndCacheMessageType(*It))
+				{
+					It.RemoveCurrent();
+				}
+			}
+
 			ConsumeInboundSegments();
 			ConsumeOutboundMessages();
 			UpdateKnownNodes();
@@ -523,6 +545,21 @@ void FUdpMessageProcessor::ConsumeInboundSegments()
 			{
 				NodeInfo.NodeId = Header.SenderNodeId;
 				NodeInfo.ProtocolVersion = Header.ProtocolVersion;
+
+				const uint16 MaxSegmentSizeV10_V17 = 2048 - UE::UdpMessaging::PacketHeaderBytes;
+				const uint16 MaxSegmentSize = GetDefault<UUdpMessagingSettings>()->GetMaxPacketSize()
+					- UE::UdpMessaging::PacketHeaderBytes;
+
+				NodeInfo.MaxSegmentSize = (Header.ProtocolVersion >= 18)
+					? MaxSegmentSize
+					: MaxSegmentSizeV10_V17;
+
+				// If USO is enabled, we need to clamp to that size here.
+				if (const uint16 SegmentOffloadPayloadSize = SocketSender->GetSegmentOffloadSize())
+				{
+					NodeInfo.MaxSegmentSize = FMath::Min(NodeInfo.MaxSegmentSize, SegmentOffloadPayloadSize);
+				}
+
 				NodeDiscoveredDelegate.ExecuteIfBound(NodeInfo.NodeId);
 				bAddedNewKnownNodes = true;
 			}
@@ -612,12 +649,12 @@ void FUdpMessageProcessor::ConsumeOneOutboundMessage(const FOutboundMessage& Out
 	if (LastSentMessage < 0)
 	{
 		// Prevent negative message ids
-		//
 		LastSentMessage = 0;
 	}
+
 	for (FNodeInfo* RecipientNodeInfo : Recipients)
 	{
-		UE_LOG(LogUdpMessaging, Verbose, TEXT("Passing %" INT64_FMT " byte message to be segment-sent %d to %s with id %s"),
+		UE_LOG(LogUdpMessaging, VeryVerbose, TEXT("Passing %" INT64_FMT " byte message to be segment-sent %d to %s with id %s"),
 				OutboundMessage.SerializedMessage->TotalSize(),
 				LastSentMessage,
 				*RecipientNodeInfo->Endpoint.ToString(),
@@ -629,9 +666,10 @@ void FUdpMessageProcessor::ConsumeOneOutboundMessage(const FOutboundMessage& Out
 			continue;
 		}
 
-		TSharedPtr<FUdpMessageSegmenter> Segmenter = MakeShared<FUdpMessageSegmenter>(OutboundMessage.SerializedMessage.ToSharedRef(), UDP_MESSAGING_SEGMENT_SIZE);
-		RecipientNodeInfo->Segmenters.Add(
-			LastSentMessage, Segmenter);
+		TSharedPtr<FUdpMessageSegmenter> Segmenter = MakeShared<FUdpMessageSegmenter>(
+			OutboundMessage.SerializedMessage.ToSharedRef(), RecipientNodeInfo->MaxSegmentSize);
+
+		RecipientNodeInfo->Segmenters.Add(LastSentMessage, Segmenter);
 		if (bIsReliable)
 		{
 			RecipientNodeInfo->ReliableWorkQueue.Add(LastSentMessage);
@@ -763,6 +801,15 @@ void FUdpMessageProcessor::ProcessDataSegment(FInboundSegment& Segment, FNodeInf
 	NodeInfo.Statistics.TotalBytesReceived += DataChunk.Data.Num();
 	NodeInfo.Statistics.PacketsReceived++;
 	ReassembledMessage->Reassemble(DataChunk.SegmentNumber, DataChunk.SegmentOffset, DataChunk.Data, CurrentTime);
+
+	if (DataChunk.SegmentNumber == 0)
+	{
+		if (!LookupAndCacheMessageType(ReassembledMessage))
+		{
+			PendingTypeLookup.Add(ReassembledMessage);
+		}
+	}
+
 	FOnInboundTransferDataUpdated& ReassemblerUpdated = UE::Private::MessageProcessor::OnReassemblerUpdated();
 	if(ReassemblerUpdated.IsBound())
 	{
@@ -770,7 +817,7 @@ void FUdpMessageProcessor::ProcessDataSegment(FInboundSegment& Segment, FNodeInf
 			{
 				NodeInfo.NodeId,
 				DataChunk.MessageId,
-				UDP_MESSAGING_SEGMENT_SIZE * ReassembledMessage->GetPendingSegmentsCount(),
+				NodeInfo.MaxSegmentSize * ReassembledMessage->GetPendingSegmentsCount(),
 				ReassembledMessage->GetReceivedBytes(),
 				ReassembledMessage->IsReliable(),
 				EnumHasAnyFlags(ReassembledMessage->GetFlags(), EMessageFlags::Traced),
@@ -917,30 +964,37 @@ void FUdpMessageProcessor::ProcessUnknownSegment(FInboundSegment& Segment, FNode
 	UE_LOG(LogUdpMessaging, Verbose, TEXT("Received unknown segment type '%i' from %s"), SegmentType, *Segment.Sender.ToText().ToString());
 }
 
-void FUdpMessageProcessor::LookupAndCacheMessageType(TSharedPtr<FUdpReassembledMessage, ESPMode::ThreadSafe>& ReassembledMessage)
+bool FUdpMessageProcessor::LookupAndCacheMessageType(TSharedPtr<FUdpReassembledMessage, ESPMode::ThreadSafe>& ReassembledMessage)
 {
-	if (!ReassembledMessage->HasFirstSegment())
+	if (ReassembledMessage->GetMessageTypeInfo().IsValid())
 	{
-		return;
+		return true;
 	}
-	if (ReassembledMessage->GetMessageTypeInfo() == nullptr)
+
+	if (!ReassembledMessage->HasFirstSegment())
 	{
-		FNameOrFTopLevel AssetPath = FUdpDeserializedMessage::PeekMessageTypeInfoName(*ReassembledMessage);
-		FString PathAsString = Visit([](auto&& Path) -> FString { return Path.ToString(); }, AssetPath);
+		return false;
+	}
 
-		if (TWeakObjectPtr<UScriptStruct>* TypeInfo = CachedTypeInfoMap.Find(PathAsString))
-		{
-			ReassembledMessage->SetMessageTypeInfo(*TypeInfo);
-		}
-		else if (!UE::IsSavingPackage() && !IsGarbageCollecting())
-		{
-			// Otherwise we have to look up the object by calling FindObjectSafe.  This can fail in GC and package save.
-			// Thus we only do this in not saving / GC cases.
-			TWeakObjectPtr<UScriptStruct> Obj = FUdpDeserializedMessage::ResolvePath(AssetPath);
-			CachedTypeInfoMap.Add(PathAsString, Obj);
-			ReassembledMessage->SetMessageTypeInfo(MoveTemp(Obj));
-		}
+	FNameOrFTopLevel AssetPath = FUdpDeserializedMessage::PeekMessageTypeInfoName(*ReassembledMessage);
+	FString PathAsString = Visit([](auto&& Path) -> FString { return Path.ToString(); }, AssetPath);
+
+	if (TWeakObjectPtr<UScriptStruct>* TypeInfo = CachedTypeInfoMap.Find(PathAsString))
+	{
+		ReassembledMessage->SetMessageTypeInfo(*TypeInfo);
+		return true;
 	}
+	else if (!UE::IsSavingPackage() && !IsGarbageCollecting())
+	{
+		// Otherwise we have to look up the object by calling FindObjectSafe.  This can fail in GC and package save.
+		// Thus we only do this in not saving / GC cases.
+		TWeakObjectPtr<UScriptStruct> Obj = FUdpDeserializedMessage::ResolvePath(AssetPath);
+		CachedTypeInfoMap.Add(PathAsString, Obj);
+		ReassembledMessage->SetMessageTypeInfo(MoveTemp(Obj));
+		return true;
+	}
+
+	return false;
 }
 
 void FUdpMessageProcessor::TryDeliverMessage(const TSharedPtr<FUdpReassembledMessage, ESPMode::ThreadSafe>& ReassembledMessage, FNodeInfo& NodeInfo)
@@ -996,7 +1050,7 @@ void FUdpMessageProcessor::RemoveDeadNodes()
 		FGuid& NodeId = It->Key;
 		FNodeInfo& NodeInfo = It->Value;
 
-		if ((NodeId.IsValid()) && HasTimedOut(NodeInfo.LastSegmentReceivedTime,CurrentTime))
+		if ((NodeId.IsValid()) && HasTimedOut(NodeInfo.LastSegmentReceivedTime, CurrentTime))
 		{
 			UE_LOG(LogUdpMessaging, Display, TEXT("FUdpMessageProcessor::UpdateKnownNodes: Removing Node %s (%s)"), *NodeInfo.NodeId.ToString(), *NodeInfo.Endpoint.ToString());
 			NodeLostDelegate.ExecuteIfBound(NodeId);
@@ -1037,39 +1091,51 @@ void FUdpMessageProcessor::SendKnownNodesToKnownNodes()
 		Header.SegmentType = EUdpMessageSegments::Mesh;
 	};
 
+	// First pass: gather everything to be sent.
 	TArray<FIPv4Endpoint> Endpoints, OutEndpoints;
 	TArray<FGuid> EndpointIds, OutIds;
-	for (const auto& NodePair : KnownNodes)
+	for (const TPair<FGuid, FNodeInfo>& NodePair : KnownNodes)
 	{
 		EndpointIds.Add(NodePair.Key);
 		Endpoints.Add(NodePair.Value.Endpoint);
 	}
 
-	// We use TArrayView to slice parts of the array for sending in parts (in case we exceed segment sending limits).
-	TArrayView<FIPv4Endpoint> KnownEndpointsView(Endpoints);
-	TArrayView<FGuid> KnownIdView(EndpointIds);
-
-	// We can send ~80 addresses in one segment. Warn the user when we hit this value. However, we still handle meshing large networks.
-	const int32 NumEndpointsCanSend = (UDP_MESSAGING_SEGMENT_SIZE - sizeof(FUdpMessageSegment::FHeader)) / (sizeof(FIPv4Endpoint) + sizeof(FGuid)) - 1;
-	if (KnownEndpointsView.Num() > NumEndpointsCanSend)
+	// Second pass: send to everyone (segment size can vary per peer node).
+	for (const TPair<FGuid, FNodeInfo>& NodePair : KnownNodes)
 	{
-		UE_LOG(LogUdpMessaging, Warning, TEXT("FUdpMessageProcessor::SendKnownNodesToKnownNodes large number of endpoints (%d) to share for meshing udp transport."), KnownEndpointsView.Num());
-	}
+		// We use TArrayView to slice parts of the array for sending (potentially multiple segments).
+		TArrayView<FIPv4Endpoint> KnownEndpointsView(Endpoints);
+		TArrayView<FGuid> KnownIdView(EndpointIds);
 
-	int32 Index = 0;
-	while (Index < KnownEndpointsView.Num())
-	{
-		const uint32 NextBlockItemCount = FMath::Min<uint32>(Index + NumEndpointsCanSend, KnownEndpointsView.Num() - Index);
-		TArrayView<FIPv4Endpoint> SlicedEndpointView = KnownEndpointsView.Slice(Index, NextBlockItemCount);
-		TArrayView<FGuid>	      SlicedIdView = KnownIdView.Slice(Index, NextBlockItemCount);
+		const uint16 MeshSegmentFixedBytes = FUdpMessageSegment::FHeader::SerializedBytes
+			+ sizeof(decltype(OutEndpoints)::SizeType)
+			+ sizeof(decltype(OutIds)::SizeType);
 
-		OutEndpoints = SlicedEndpointView;
-		OutIds = SlicedIdView;
+		if (!ensure(MeshSegmentFixedBytes < NodePair.Value.MaxSegmentSize))
+		{
+			continue;
+		}
 
-		UE_LOG(LogUdpMessaging, VeryVerbose, TEXT("FUdpMessageProcessor::SendKnownNodesToKnownNodes Sending updated known nodes."));
+		const uint16 MaxMeshPayload = NodePair.Value.MaxSegmentSize - MeshSegmentFixedBytes;
+		const int32 NumEndpointsCanSend = MaxMeshPayload / (sizeof(FIPv4Endpoint) + sizeof(FGuid)) - 1;
+
+		if (!ensure(NumEndpointsCanSend > 0))
+		{
+			continue;
+		}
 
-		for (const FIPv4Endpoint& Endpoint : Endpoints)
+		int32 Index = 0;
+		while (Index < KnownEndpointsView.Num())
 		{
+			const int32 NextBlockItemCount = FMath::Min(NumEndpointsCanSend, KnownEndpointsView.Num() - Index);
+			TArrayView<FIPv4Endpoint> SlicedEndpointView = KnownEndpointsView.Slice(Index, NextBlockItemCount);
+			TArrayView<FGuid>	      SlicedIdView = KnownIdView.Slice(Index, NextBlockItemCount);
+
+			OutEndpoints = SlicedEndpointView;
+			OutIds = SlicedIdView;
+
+			UE_LOG(LogUdpMessaging, VeryVerbose, TEXT("FUdpMessageProcessor::SendKnownNodesToKnownNodes Sending updated known nodes."));
+
 			TSharedRef<FArrayWriter, ESPMode::ThreadSafe> Writer = MakeShared<FArrayWriter, ESPMode::ThreadSafe>();
 			{
 				*Writer << Header;
@@ -1077,13 +1143,16 @@ void FUdpMessageProcessor::SendKnownNodesToKnownNodes()
 				*Writer << OutIds;
 			}
 
+			const FIPv4Endpoint& Endpoint = NodePair.Value.Endpoint;
 			if (!SocketSender->Send(Writer, Endpoint))
 			{
 				UE_LOG(LogUdpMessaging, Warning, TEXT("FUdpMessageProcessor::SendKnownNodesToKnownNodes failed to share endpoint information to %s."), *Endpoint.ToString());
 			}
+
+			Index += NumEndpointsCanSend;
 		}
-		Index += NumEndpointsCanSend;
 	}
+
 	bAddedNewKnownNodes = false;
 }
 
@@ -1199,92 +1268,177 @@ TSharedPtr<FUdpMessageSegmenter>* FUdpMessageProcessor::GetInitializedSegmenter(
 		return nullptr;
 	}
 
-	return  FoundSegmenter;
+	return FoundSegmenter;
 }
 
-FUdpMessageSegment::FDataChunk FUdpMessageProcessor::GetDataChunk(FNodeInfo& NodeInfo, TSharedPtr<FUdpMessageSegmenter>& Segmenter, int32 MessageId)
+bool FUdpMessageProcessor::FillDataChunk(FUdpMessageSegment::FDataChunk& OutDataChunk, FNodeInfo& NodeInfo, TSharedPtr<FUdpMessageSegmenter>& Segmenter, int32 MessageId)
 {
-	FUdpMessageSegment::FDataChunk DataChunk;
-
 	// Track the segments we sent as we'll update the segmenter to keep track
-	TConstSetBitIterator<> BIt(Segmenter->GetPendingSendSegments());
-	Segmenter->GetPendingSegment(BIt.GetIndex(), DataChunk.Data);
-	DataChunk.SegmentNumber = BIt.GetIndex();
+	uint32 OutSegmentNumber;
+	if (!Segmenter->GetNextPendingSegment(OutDataChunk.Data, OutSegmentNumber))
+	{
+		return false;
+	}
+
+	OutDataChunk.MessageId = MessageId;
+	OutDataChunk.MessageSize = Segmenter->GetMessageSize();
+	OutDataChunk.MessageFlags = Segmenter->GetMessageFlags();
+	OutDataChunk.SegmentNumber = OutSegmentNumber;
+	OutDataChunk.SegmentOffset = OutSegmentNumber * static_cast<uint64>(Segmenter->GetSegmentPayloadSize());
+	OutDataChunk.TotalSegments = Segmenter->GetSegmentCount();
+	OutDataChunk.Sequence = 0; // This should be kept 0 for legacy support.
+
+	return true;
+}
+
+int32 FUdpMessageProcessor::ClampMaxSegments(int32 MaxSegments, const TSharedPtr<FUdpMessageSegmenter>& Segmenter, const FNodeInfo& NodeInfo) const
+{
+	const int32 SegmentOffloadPayloadSize = SocketSender->GetSegmentOffloadSize();
+	if (SegmentOffloadPayloadSize > 0
+		&& NodeInfo.MaxSegmentSize >= SegmentOffloadPayloadSize
+		&& Segmenter->GetTotalSegmentSize() == SegmentOffloadPayloadSize)
+	{
+		// Hardware acceleration frequently limits the pre-segmented size to 64K.
+		const int32 SenderMaxSubmitSize = UINT16_MAX;
+		const int32 SenderMaxSubmitSegments = SenderMaxSubmitSize / SegmentOffloadPayloadSize;
+		MaxSegments = FMath::Min(MaxSegments, SenderMaxSubmitSegments);
+
+		// "UDP_SEGMENT [...] The segment size must be chosen such that at most 64 datagrams are sent in a single call" - man udp(7)
+		const int32 SenderMaxGeneratedSegments = 64;
+		MaxSegments = FMath::Min(MaxSegments, SenderMaxGeneratedSegments);
+
+		const int32 CvarMaxGeneratedSegments = CVarUdpMessagingMaxOffloadSegments.GetValueOnAnyThread();
+		MaxSegments = FMath::Clamp(MaxSegments, 1, CvarMaxGeneratedSegments);
+	}
+	else
+	{
+		// Segmentation offload not supported; only send a single segment at a time.
+		MaxSegments = 1;
+	}
 
-	DataChunk.MessageId = MessageId;
-	DataChunk.MessageFlags = Segmenter->GetMessageFlags();
-	DataChunk.MessageSize = Segmenter->GetMessageSize();
-	DataChunk.SegmentOffset = UDP_MESSAGING_SEGMENT_SIZE * DataChunk.SegmentNumber;
-	DataChunk.Sequence = 0; // This should be kept 0 for legacy support.
-	DataChunk.TotalSegments = Segmenter->GetSegmentCount();
-	return DataChunk;
+	return MaxSegments;
 }
 
-FSentSegmentInfo FUdpMessageProcessor::SendNextSegmentForMessageId(FNodeInfo& NodeInfo, FUdpMessageSegment::FHeader& Header, int32 MessageId)
+FSentSegmentInfo FUdpMessageProcessor::SendNextSegmentsForMessageId(
+	FNodeInfo& NodeInfo,
+	FUdpMessageSegment::FHeader& Header,
+	int32 MessageId,
+	int32 MaxSegments /* = 1 */
+)
 {
 	TSharedPtr<FUdpMessageSegmenter>* FoundSegmenter = GetInitializedSegmenter(NodeInfo, MessageId);
 	if (!FoundSegmenter)
 	{
 		// It's possible that a segmenter was fully ack and removed from the segmenter list.
-		return {MessageId};
+		return FSentSegmentInfo(MessageId);
 	}
 
 	TSharedPtr<FUdpMessageSegmenter>& Segmenter = *FoundSegmenter;
 
 	FSentSegmentInfo SentInfo(MessageId);
+	TArray<FSentData, TInlineAllocator<64>> SentDatas;
 	if (Segmenter->IsInitialized() && Segmenter->NeedSending(CurrentTime))
 	{
-		FUdpMessageSegment::FDataChunk DataChunk = GetDataChunk(NodeInfo, Segmenter, MessageId);
-
 		if (Header.ProtocolVersion != Segmenter->GetProtocolVersion())
 		{
-			UE_LOG(LogUdpMessaging, Error, TEXT("FUdpMessageProcessor::UpdateSegmenters: Node %s initially asked for protocol version %d but changed to %d. Cannot deliver message %d"), *NodeInfo.Endpoint.ToString(), Segmenter->GetProtocolVersion(), Header.ProtocolVersion, MessageId);
+			UE_LOG(LogUdpMessaging, Error,
+				TEXT("FUdpMessageProcessor::SendNextSegmentsForMessageId: Node %s initially asked for protocol version %d but changed to %d. Cannot deliver message %d"),
+				*NodeInfo.Endpoint.ToString(), Segmenter->GetProtocolVersion(), Header.ProtocolVersion, MessageId);
 			NodeInfo.Segmenters.Remove(MessageId);
-			return SentInfo;
+			return FSentSegmentInfo(MessageId);
 		}
 
+		const bool bIsReliable = EnumHasAnyFlags(Segmenter->GetMessageFlags(), EMessageFlags::Reliable);
+
+		MaxSegments = ClampMaxSegments(MaxSegments, Segmenter, NodeInfo);
+
 		TSharedRef<FArrayWriter, ESPMode::ThreadSafe> Writer = MakeShared<FArrayWriter, ESPMode::ThreadSafe>();
+		Writer->Reserve(NodeInfo.MaxSegmentSize * MaxSegments);
+
+		FUdpMessageSegment::FDataChunk DataChunk;
+
+		for (int32 SendIdx = 0; SendIdx < MaxSegments; ++SendIdx)
 		{
+			if (Segmenter->IsSendingComplete())
+			{
+				break;
+			}
+
+			if (!FillDataChunk(DataChunk, NodeInfo, Segmenter, MessageId))
+			{
+				UE_LOG(LogUdpMessaging, Error,
+					TEXT("FUdpMessageProcessor::SendNextSegmentsForMessageId: FillDataChunk failed"));
+				break;
+			}
+
+			const int64 WriteStartPos = Writer->Tell();
 			*Writer << Header;
 			DataChunk.Serialize(*Writer, Header.ProtocolVersion);
-		}
+			const int64 WriteEndPos = Writer->Tell();
+			const int64 BytesWritten = WriteEndPos - WriteStartPos;
+
+			UE_LOG(LogUdpMessaging, VeryVerbose, TEXT("FUdpMessageProcessor::UpdateSegmenters: Sending msg %d as segment %d/%d of %" INT64_FMT " bytes to %s"),
+				MessageId,
+				DataChunk.SegmentNumber + 1,
+				DataChunk.TotalSegments,
+				Segmenter->GetMessageSize(),
+				*NodeInfo.NodeId.ToString());
 
-		UE_LOG(LogUdpMessaging, Verbose, TEXT("FUdpMessageProcessor::UpdateSegmenters: Sending msg %d as segment %d/%d of %" INT64_FMT " bytes to %s"),
-			   DataChunk.MessageId,
-			   DataChunk.SegmentNumber + 1,
-			   DataChunk.TotalSegments,
-			   Segmenter->GetMessageSize(),
-			   *NodeInfo.NodeId.ToString());
+			// This is optimistic; we'll undo it later if the batched send fails.
+			Segmenter->MarkAsSent(DataChunk.SegmentNumber);
 
-		if (!SocketSender->Send(Writer, NodeInfo.Endpoint))
-		{
-			SentInfo.bSendSocketError = true;
-			return MoveTemp(SentInfo);
+			SentDatas.Emplace(
+				MessageId,
+				DataChunk.SegmentNumber,
+				++NodeInfo.SequenceId,
+				bIsReliable,
+				CurrentTime
+			);
 		}
 
-		Segmenter->MarkAsSent(DataChunk.SegmentNumber);
-
-		NodeInfo.Statistics.PacketsSent++;
-		SentInfo.BytesSent += Writer->Num();
-		SentInfo.SequenceNumber = ++NodeInfo.SequenceId;
-		SentInfo.bIsReliable = EnumHasAnyFlags(Segmenter->GetMessageFlags(), EMessageFlags::Reliable);
-		SentInfo.bRequiresRequeue = !Segmenter->IsSendingComplete() || SentInfo.bIsReliable;
-		SentInfo.bFullySent = Segmenter->IsSendingComplete();
-		SentInfo.bIsTraced = EnumHasAnyFlags(Segmenter->GetMessageFlags(), EMessageFlags::Traced);
+		const bool bAllowSegmentation_true = true;
+		if (SocketSender->Send(Writer, NodeInfo.Endpoint, bAllowSegmentation_true))
+		{
+			// FIXME?: NeedSending above implicitly updates the LastSentTime, even if send fails
+			Segmenter->UpdateSentTime(CurrentTime);
+
+			SentInfo.SequenceNumber = NodeInfo.SequenceId;
+			SentInfo.BytesSent = Writer->Tell();
+			SentInfo.SegmentsSent = SentDatas.Num();
+			SentInfo.bIsReliable = bIsReliable;
+			SentInfo.bRequiresRequeue = !Segmenter->IsSendingComplete() || SentInfo.bIsReliable;
+			SentInfo.bFullySent = Segmenter->IsSendingComplete();
+			SentInfo.bIsTraced = EnumHasAnyFlags(Segmenter->GetMessageFlags(), EMessageFlags::Traced);
 
-		Segmenter->UpdateSentTime(CurrentTime);
+			NodeInfo.Statistics.PacketsSent += SentDatas.Num();
 
-		NodeInfo.InflightSegments.Add(
-			NodeInfo.Remap(MessageId,DataChunk.SegmentNumber),
-			SentInfo.AsSentData(MessageId,DataChunk.SegmentNumber,CurrentTime));
+			for (FSentData& SentData : SentDatas)
+			{
+				NodeInfo.InflightSegments.Add(
+					NodeInfo.Remap(MessageId, SentData.SegmentNumber),
+					SentData);
+			}
 
-		SendSegmenterStatsToListeners(MessageId, NodeInfo.NodeId, Segmenter);
+			SendSegmenterStatsToListeners(MessageId, NodeInfo.NodeId, Segmenter);
 
-		if (Segmenter->IsSendingComplete() && !SentInfo.bIsReliable)
+			if (Segmenter->IsSendingComplete() && !bIsReliable)
+			{
+				// Not reliably sent so we don't need to wait for acks.
+				UE_LOG(LogUdpMessaging, VeryVerbose, TEXT("FUdpMessageProcessor::UpdateSegmenters: Finished with message segmenter for %s"), *NodeInfo.NodeId.ToString());
+				NodeInfo.Segmenters.Remove(MessageId);
+			}
+		}
+		else
 		{
-			// Not reliably sent so we don't need to wait for acks.
-			UE_LOG(LogUdpMessaging, VeryVerbose, TEXT("FUdpMessageProcessor::UpdateSegmenters: Finished with message segmenter for %s"), *NodeInfo.NodeId.ToString());
-			NodeInfo.Segmenters.Remove(MessageId);
+			for (const FSentData& SentData : SentDatas)
+			{
+				// Roll back the sent state of any segments we failed to send.
+				Segmenter->MarkForRetransmission(SentData.SegmentNumber);
+			}
+
+			SentInfo.bSendSocketError = true;
+			SentInfo.SegmentsSent = 0;
+			return MoveTemp(SentInfo);
 		}
 	}
 	else
@@ -1320,30 +1474,32 @@ TOptional<FSentSegmentInfoArray> FUdpMessageProcessor::DrainQueue(FNodeInfo& Nod
 	if (bDrainOverflow)
 	{
 		// Check for acknowledged segments first.
-		while (NodeInfo.CanSendSegments() && !NodeInfo.OverflowForPendingAck.IsEmpty())
+		while (DrainCount > 0 && NodeInfo.CanSendSegments() && !NodeInfo.OverflowForPendingAck.IsEmpty())
 		{
-			FSentSegmentInfo Info = SendNextSegmentForMessageId(NodeInfo, Header, NodeInfo.OverflowForPendingAck.Pop());
+			const int32 MessageId = NodeInfo.OverflowForPendingAck.Pop();
+			FSentSegmentInfo Info = SendNextSegmentsForMessageId(NodeInfo, Header, MessageId, DrainCount);
 			if (Info.bSendSocketError)
 			{
 				return {};
 			}
 			ProcessedMessages.Emplace(MoveTemp(Info));
-			DrainCount--;
+			DrainCount -= FMath::Max(1u, Info.SegmentsSent);
 		}
 	}
-	
+
 	// Process messages in the work queue.
 	while (DrainCount > 0 && NodeInfo.CanSendSegments() && !WorkQueue.IsEmpty())
 	{
-		int32 MessageId = WorkQueue.PopFrontValue();
-		FSentSegmentInfo Info = SendNextSegmentForMessageId(NodeInfo, Header, MessageId);
+		const int32 MessageId = WorkQueue.PopFrontValue();
+		FSentSegmentInfo Info = SendNextSegmentsForMessageId(NodeInfo, Header, MessageId, DrainCount);
 		if (Info.bSendSocketError)
 		{
 			return {};
 		}
 		ProcessedMessages.Emplace(MoveTemp(Info));
-		DrainCount--;
+		DrainCount -= FMath::Max(1u, Info.SegmentsSent);
 	}
+
 	return ProcessedMessages;
 }
 
@@ -1434,7 +1590,7 @@ int32 FUdpMessageProcessor::UpdateSegmenters(FNodeInfo& NodeInfo)
 	}
 
 	NodeInfo.Statistics.PacketsInFlight = NodeInfo.InflightSegments.Num();
-	NodeInfo.Statistics.BytesInflight = NodeInfo.InflightSegments.Num() * UDP_MESSAGING_SEGMENT_SIZE;
+	NodeInfo.Statistics.BytesInflight = NodeInfo.InflightSegments.Num() * NodeInfo.MaxSegmentSize;
 
 	return BytesSent;
 }
@@ -1452,6 +1608,16 @@ bool FUdpMessageProcessor::UpdateReassemblers(FNodeInfo& NodeInfo)
 		EUdpMessageSegments::AcknowledgeSegments			// Header.SegmentType
 	};
 
+	const uint16 MaxAckSegmentPayload = NodeInfo.MaxSegmentSize
+		- FUdpMessageSegment::FHeader::SerializedBytes
+		- FUdpMessageSegment::FAcknowledgeSegmentsChunk::FixedSerializedBytes;
+
+	const int32 MaxAckNum = (MaxAckSegmentPayload / sizeof(uint32));
+	if (!ensure(MaxAckNum > 0))
+	{
+		return true;
+	}
+
 	const uint32 MaxSendRate = GetMaxSendRate();
 	double DeltaSeconds = DeltaTime.GetTotalSeconds();
 	const int32 MaxSendRateDelta = MaxSendRate * DeltaSeconds;
@@ -1459,7 +1625,9 @@ bool FUdpMessageProcessor::UpdateReassemblers(FNodeInfo& NodeInfo)
 
 	for (TMap<int32, TSharedPtr<FUdpReassembledMessage, ESPMode::ThreadSafe>>::TIterator It(NodeInfo.ReassembledMessages); It; ++It)
 	{
-		TSharedPtr<FUdpReassembledMessage, ESPMode::ThreadSafe>& ReassembledMessage = It.Value();
+		TSharedPtr<FUdpReassembledMessage, ESPMode::ThreadSafe>& ReassembledMessagePtr = It.Value();
+		// Believe it or not, TSharedPtr operator-> otherwise makes a big showing here.
+		FUdpReassembledMessage* ReassembledMessage = It.Value().Get();
 
 		if (ReassembledMessage->IsDelivered())
 		{
@@ -1471,7 +1639,7 @@ bool FUdpMessageProcessor::UpdateReassemblers(FNodeInfo& NodeInfo)
 		int BytesSent = 0;
 		while (ReassembledMessage->HasPendingAcknowledgements() && BytesSent < MaxNodeByteSend)
 		{
-			TArray<uint32> PendingAcknowledgments = ReassembledMessage->GetPendingAcknowledgments();
+			TArray<uint32> PendingAcknowledgments = ReassembledMessage->TakePendingAcknowledgments(MaxAckNum);
 			const int32 AckCount = PendingAcknowledgments.Num();
 
 			FUdpMessageSegment::FAcknowledgeSegmentsChunk AcknowledgeChunk(It.Key()/*MessageId*/, MoveTemp(PendingAcknowledgments)/*Segments*/);
@@ -1496,12 +1664,10 @@ bool FUdpMessageProcessor::UpdateReassemblers(FNodeInfo& NodeInfo)
 				*NodeInfo.NodeId.ToString());
 		}
 
-		LookupAndCacheMessageType(ReassembledMessage);
-
 		// Try to deliver completed message that couldn't be delivered the first time around
 		if (ReassembledMessage->IsComplete() && !(ReassembledMessage->IsPendingDelivery() || ReassembledMessage->IsDelivered()))
 		{
-			TryDeliverMessage(ReassembledMessage, NodeInfo);
+			TryDeliverMessage(ReassembledMessagePtr, NodeInfo);
 		}
 
 		// Remove stale reassembled message if they aren't reliable or are marked delivered
@@ -1520,6 +1686,7 @@ bool FUdpMessageProcessor::UpdateReassemblers(FNodeInfo& NodeInfo)
 			if (!ReassembledMessage->IsPendingDelivery())
 			{
 				It.RemoveCurrent();
+				PendingTypeLookup.Remove(ReassembledMessagePtr);
 			}
 		}
 	}
diff --git a/Engine/Plugins/Messaging/UdpMessaging/Source/UdpMessaging/Private/Transport/UdpMessageProcessor.h b/Engine/Plugins/Messaging/UdpMessaging/Source/UdpMessaging/Private/Transport/UdpMessageProcessor.h
--- a/Engine/Plugins/Messaging/UdpMessaging/Source/UdpMessaging/Private/Transport/UdpMessageProcessor.h
+++ b/Engine/Plugins/Messaging/UdpMessaging/Source/UdpMessaging/Private/Transport/UdpMessageProcessor.h
@@ -18,6 +18,7 @@
 #include "UObject/WeakObjectPtr.h"
 #include "UObject/WeakObjectPtrTemplates.h"
 
+#include "Shared/UdpMessagingSettings.h"
 #include "UdpMessageSegmenter.h"
 #include "UdpMessagingPrivate.h"
 
@@ -75,6 +76,7 @@ struct FSentSegmentInfo
 	int32 MessageId;
 	uint64 SequenceNumber = 0;
 	uint32 BytesSent = 0;
+	uint32 SegmentsSent = 0;
 
 	bool bIsReliable = false;
 	bool bRequiresRequeue = false;
@@ -83,18 +85,6 @@ struct FSentSegmentInfo
 	bool bIsTraced = false;
 	bool bIsReady = false;
 
-	/** Converts FSentSegmentInfo into a FSentData struct. */
-	FSentData AsSentData(int32 InMessageId, uint32 SegmentNumber, const FDateTime& CurrentTime)
-	{
-		return {
-			InMessageId,
-			SegmentNumber,
-			SequenceNumber,
-			bIsReliable,
-			CurrentTime
-		};
-	}
-
 	bool WasSent() const
 	{
 		return BytesSent > 0;
@@ -149,6 +139,9 @@ class FUdpMessageProcessor
 		/** Holds the protocol version this node is communicating with */
 		uint8 ProtocolVersion;
 
+		/** The maximum datagram payload (UDP SDU) this node is capable of receiving. Clamped to USO size if enabled. */
+		uint16 MaxSegmentSize;
+
 		/** Holds the collection of reassembled messages. */
 		TMap<int32, TSharedPtr<FUdpReassembledMessage, ESPMode::ThreadSafe>> ReassembledMessages;
 
@@ -203,11 +196,11 @@ class FUdpMessageProcessor
 			//
 			if (MessageId >= SegmentId)
 			{
-				return MessageId * MessageId + MessageId + SegmentId;
+				return (uint64)MessageId * MessageId + MessageId + SegmentId;
 			}
 			else
 			{
-				return SegmentId * SegmentId + MessageId;
+				return (uint64)SegmentId * SegmentId + MessageId;
 			}
 		}
 
@@ -379,7 +372,7 @@ class FUdpMessageProcessor
 				return bIsLost;
 			});
 
-			Statistics.TotalBytesLost += SegmentsLost * UDP_MESSAGING_SEGMENT_SIZE;
+			Statistics.TotalBytesLost += SegmentsLost * MaxSegmentSize;
 			Statistics.PacketsLost += SegmentsLost;
 			return SegmentsLost;
 		}
@@ -468,18 +461,17 @@ class FUdpMessageProcessor
 		{ }
 	};
 
-	FUdpMessageSegment::FDataChunk GetDataChunk(FNodeInfo& NodeInfo, TSharedPtr<FUdpMessageSegmenter>& Segmenter, int32 MessageId);
+	bool FillDataChunk(FUdpMessageSegment::FDataChunk& OutDataChunk, FNodeInfo& NodeInfo, TSharedPtr<FUdpMessageSegmenter>& Segmenter, int32 MessageId);
 
 	/** Returns an initialized segment ready for sending or ready to receive ack. */
 	TSharedPtr<FUdpMessageSegmenter>* GetInitializedSegmenter(FNodeInfo& NodeInfo, int32 MessageId);
 
-	/** Get a Sent Data structure for the outbound datachunk. */
-	FSentData GetSentDataFromInfo(const FSentSegmentInfo &Info, const FUdpMessageSegment::FDataChunk& Chunk);
+	int32 ClampMaxSegments(int32 MaxSegments, const TSharedPtr<FUdpMessageSegmenter>& Segmenter, const FNodeInfo& NodeInfo) const;
 
-	/** Send one segment out for the given message id.  Send state is provided by the return value. */
-	FSentSegmentInfo SendNextSegmentForMessageId(FNodeInfo& NodeInfo, FUdpMessageSegment::FHeader& Header, int32 MessageId);
+	/** Send one or more segments out for the given message id. Send state is provided by the return value. */
+	FSentSegmentInfo SendNextSegmentsForMessageId(FNodeInfo& NodeInfo, FUdpMessageSegment::FHeader& Header, int32 MessageId, int32 MaxSegments = 1);
 
-	/** Remove any nodes we have not heard from in  while. */
+	/** Remove any nodes we have not heard from in a while. */
 	void RemoveDeadNodes();
 
 	/** Update network statistics table that callers can use to gather info about running connections. */
@@ -811,7 +803,7 @@ protected:
 	/**
 	 * Lookup the reassembled message type and try to determine the source struct type.
 	 */
-	void LookupAndCacheMessageType(TSharedPtr<FUdpReassembledMessage, ESPMode::ThreadSafe>& ReassembledMessage);
+	bool LookupAndCacheMessageType(TSharedPtr<FUdpReassembledMessage, ESPMode::ThreadSafe>& ReassembledMessage);
 
 	/**
 	 * Removes the specified node from the list of known remote endpoints.
@@ -1023,6 +1015,9 @@ private:
 	/** Stores a cache of discovered UScriptStruct types found during message processing. */
 	TMap<FString, TWeakObjectPtr<UScriptStruct>> CachedTypeInfoMap;
 
+	/** Stores messages pending type lookup. */
+	TSet<FUdpReassembledMessageSharedPtr> PendingTypeLookup;
+
 	/** How much memory (MB) should be used before periodically logging information about nodes' memory usage. */
 	float MemoryLoggingThresholdMB = -1;
 
diff --git a/Engine/Plugins/Messaging/UdpMessaging/Source/UdpMessaging/Private/Transport/UdpMessageSegmenter.cpp b/Engine/Plugins/Messaging/UdpMessaging/Source/UdpMessaging/Private/Transport/UdpMessageSegmenter.cpp
--- a/Engine/Plugins/Messaging/UdpMessaging/Source/UdpMessaging/Private/Transport/UdpMessageSegmenter.cpp
+++ b/Engine/Plugins/Messaging/UdpMessaging/Source/UdpMessaging/Private/Transport/UdpMessageSegmenter.cpp
@@ -9,6 +9,17 @@
 /* FUdpMessageSegmenter structors
  *****************************************************************************/
 
+FUdpMessageSegmenter::FUdpMessageSegmenter(
+	const TSharedRef<FUdpSerializedMessage>& InSerializedMessage,
+	uint16 InSegmentSize
+)
+	: SegmentSize(InSegmentSize)
+	, SerializedMessage(InSerializedMessage)
+{
+	check(InSegmentSize > GetDataSegmentHeaderSize());
+}
+
+
 FUdpMessageSegmenter::~FUdpMessageSegmenter()
 {
 	if (MessageReader != nullptr)
@@ -39,16 +50,16 @@ bool FUdpMessageSegmenter::GetNextPendingSegment(TArray<uint8>& OutData, uint32&
 		return false;
 	}
 
-	if (TConstSetBitIterator<> It(PendingSendSegments); It)
+	if (TConstSetBitIterator<> It(PendingSendSegments, EarliestUnackedSegmentId); It)
 	{
 		OutSegment = It.GetIndex();
 
-		uint64 SegmentOffset = static_cast<uint64>(OutSegment) * SegmentSize;
+		uint64 SegmentOffset = static_cast<uint64>(OutSegment) * GetSegmentPayloadSize();
 		uint64 ActualSegmentSize = MessageReader->TotalSize() - SegmentOffset;
 
-		if (ActualSegmentSize > SegmentSize)
+		if (ActualSegmentSize > GetSegmentPayloadSize())
 		{
-			ActualSegmentSize = SegmentSize;
+			ActualSegmentSize = GetSegmentPayloadSize();
 		}
 
 		OutData.Reset(ActualSegmentSize);
@@ -76,12 +87,12 @@ bool FUdpMessageSegmenter::GetPendingSegment(uint32 InSegment, TArray<uint8>& Ou
 	// Max segment number for protocol 12 is INT32_MAX, if increased, this will need changing
 	if (InSegment < (uint32)PendingSendSegments.Num() && PendingSendSegments[InSegment])
 	{
-		uint64 SegmentOffset = static_cast<uint64>(InSegment) * SegmentSize;
+		uint64 SegmentOffset = static_cast<uint64>(InSegment) * GetSegmentPayloadSize();
 		uint64 ActualSegmentSize = MessageReader->TotalSize() - SegmentOffset;
 
-		if (ActualSegmentSize > SegmentSize)
+		if (ActualSegmentSize > GetSegmentPayloadSize())
 		{
-			ActualSegmentSize = SegmentSize;
+			ActualSegmentSize = GetSegmentPayloadSize();
 		}
 
 		OutData.Reset(ActualSegmentSize);
@@ -107,7 +118,7 @@ void FUdpMessageSegmenter::Initialize()
 	if (SerializedMessage->GetState() == EUdpSerializedMessageState::Complete)
 	{
 		MessageReader = SerializedMessage->CreateReader();
-		PendingSendSegmentsCount = (MessageReader->TotalSize() + SegmentSize - 1) / SegmentSize;
+		PendingSendSegmentsCount = (MessageReader->TotalSize() + GetSegmentPayloadSize() - 1) / GetSegmentPayloadSize();
 		PendingSendSegments.Init(true, PendingSendSegmentsCount);
 		if (EnumHasAnyFlags(GetMessageFlags(), EMessageFlags::Reliable))
 		{
@@ -173,6 +184,12 @@ void FUdpMessageSegmenter::MarkAsAcknowledged(uint32 Segment)
 		++AcknowledgeSegmentsCount;
 		AcknowledgeSegments[Segment] = true;
 		UE_LOG(LogUdpMessaging, Verbose, TEXT("Marked segment %d of %d as acknowledged (%d outstanding)"), Segment + 1, AcknowledgeSegments.Num(), AcknowledgeSegments.Num() - AcknowledgeSegmentsCount);
+
+		if (Segment == EarliestUnackedSegmentId)
+		{
+			const int32 NextUnackedId = AcknowledgeSegments.FindFrom(false, EarliestUnackedSegmentId);
+			EarliestUnackedSegmentId = (NextUnackedId == INDEX_NONE) ? AcknowledgeSegments.Num() : NextUnackedId;
+		}
 	}
 
 	// We may have queued a segment to be resent, if so there's now no need to resend it
@@ -270,7 +287,7 @@ bool FUdpMessageSegmenter::NeedSending(const FDateTime& CurrentTime)
 	{
 		// We have gone through a period of time where packets or acks may have been lost,
 		// so resend any segments that have yet to be acknowledged
-		for (TBitArray<>::FConstIterator BIt(AcknowledgeSegments); BIt; ++BIt)
+		for (TBitArray<>::FConstIterator BIt(AcknowledgeSegments, EarliestUnackedSegmentId); BIt; ++BIt)
 		{
 			const int32 Index = BIt.GetIndex();
 			if (BIt.GetValue() == false && !PendingSendSegments[Index])
diff --git a/Engine/Plugins/Messaging/UdpMessaging/Source/UdpMessaging/Private/Transport/UdpMessageSegmenter.h b/Engine/Plugins/Messaging/UdpMessaging/Source/UdpMessaging/Private/Transport/UdpMessageSegmenter.h
--- a/Engine/Plugins/Messaging/UdpMessaging/Source/UdpMessaging/Private/Transport/UdpMessageSegmenter.h
+++ b/Engine/Plugins/Messaging/UdpMessaging/Source/UdpMessaging/Private/Transport/UdpMessageSegmenter.h
@@ -2,6 +2,7 @@
 
 #pragma once
 
+#include "Shared/UdpMessageSegment.h"
 #include "Templates/SharedPointer.h"
 #include "Misc/DateTime.h"
 
@@ -23,25 +24,19 @@ class FUdpSerializedMessage;
 class FUdpMessageSegmenter
 {
 public:
-
-	/** Default constructor. */
-	FUdpMessageSegmenter()
-		: MessageReader(nullptr)
-		, SentNumber(0)
-	{ }
-
 	/**
 	 * Creates and initializes a new message segmenter.
 	 *
+	 * Segment size includes UdpMessaging headers (\see GetDataSegmentHeaderSize),
+	 * but not UDP/IP headers.
+	 *
 	 * @param InMessage The serialized message to segment.
+	 * @param InSegmentSize The target total segment size to divide messages into.
 	 */
-	FUdpMessageSegmenter(const TSharedRef<FUdpSerializedMessage, ESPMode::ThreadSafe>& InSerializedMessage, uint16 InSegmentSize)
-		: MessageReader(nullptr)
-		, SegmentSize(InSegmentSize)
-		, SentNumber(0)
-		, LastSentTime(0)
-		, SerializedMessage(InSerializedMessage)
-	{ }
+	FUdpMessageSegmenter(
+		const TSharedRef<FUdpSerializedMessage>& InSerializedMessage,
+		uint16 InSegmentSize
+	);
 
 	/** Destructor. */
 	~FUdpMessageSegmenter();
@@ -122,6 +117,37 @@ public:
 		return PendingSendSegments.Num();
 	}
 
+	/**
+	 * Gets the total segment size including headers.
+	 *
+	 * @return Segment size in bytes.
+	 */
+	uint16 GetTotalSegmentSize() const
+	{
+		return SegmentSize;
+	}
+
+	/**
+	 * Gets the size of the headers that are prepended to each segment.
+	 *
+	 * @return Header size in bytes.
+	 */
+	static constexpr uint16 GetDataSegmentHeaderSize()
+	{
+		return FUdpMessageSegment::FHeader::SerializedBytes
+			+ FUdpMessageSegment::FDataChunk::FixedSerializedBytes;
+	}
+
+	/**
+	 * Gets the number of actual message data bytes per segment, excluding headers (goodput).
+	 *
+	 * @return Segment payload size in bytes.
+	 */
+	uint16 GetSegmentPayloadSize() const
+	{
+		return SegmentSize - GetDataSegmentHeaderSize();
+	}
+
 	/** Initializes the segmenter. */
 	void Initialize();
 
@@ -239,8 +265,8 @@ public:
 	bool IsMessageSerializationDone() const;
 
 private:
-	/** temp hack to support new transport API. */
-	FArchive* MessageReader;
+	/** Used to seek and output portions of SerializedMessage. */
+	FArchive* MessageReader = nullptr;
 
 	/** Holds an array of bits where true indicates which segments still need to be sent. */
 	TBitArray<> PendingSendSegments;
@@ -249,16 +275,19 @@ private:
 	TBitArray<> AcknowledgeSegments;
 
 	/** Holds the number of segments that haven't been sent yet. */
-	uint32 PendingSendSegmentsCount;
+	uint32 PendingSendSegmentsCount = 0;
 
 	/** Holds the number of segments that haven't been acknowledged yet. Will always be zero for an unreliable message */
-	uint32 AcknowledgeSegmentsCount;
+	uint32 AcknowledgeSegmentsCount = 0;
+
+	/** Holds the earliest segment for which we have not received an acknowledgment. Useful to narrow search. */
+	int32 EarliestUnackedSegmentId = 0;
 
-	/** Holds the segment size. */
-	uint16 SegmentSize;
+	/** Holds the segment size (including the FHeader and FDataChunk fields). */
+	uint16 SegmentSize = 0;
 
 	/** Holds the number of time we sent the segments */
-	uint16 SentNumber;
+	uint16 SentNumber = 0;
 
 	/** Holds the time at which we last sent */
 	FDateTime LastSentTime;
diff --git a/Engine/Plugins/Messaging/UdpMessaging/Source/UdpMessaging/Private/Transport/UdpMessageTransport.cpp b/Engine/Plugins/Messaging/UdpMessaging/Source/UdpMessaging/Private/Transport/UdpMessageTransport.cpp
--- a/Engine/Plugins/Messaging/UdpMessaging/Source/UdpMessaging/Private/Transport/UdpMessageTransport.cpp
+++ b/Engine/Plugins/Messaging/UdpMessaging/Source/UdpMessaging/Private/Transport/UdpMessageTransport.cpp
@@ -272,6 +272,8 @@ void JoinedToGroup(const FIPv4Endpoint& UnicastEndpoint, const FIPv4Endpoint& Mu
 
 bool FUdpMessageTransport::StartTransport(IMessageTransportHandler& Handler)
 {
+	const UUdpMessagingSettings* CurrentSettings = GetDefault<UUdpMessagingSettings>();
+
 	// Set the handler even if initialization fails. This allows tries for reinitialization using the same handler.
 	TransportHandler = &Handler;
 
@@ -295,7 +297,7 @@ bool FUdpMessageTransport::StartTransport(IMessageTransportHandler& Handler)
 		UE_LOG(LogUdpMessaging, Display, TEXT("Unicast socket bound to '%s:%d'."), *UnicastEndpoint.Address.ToString(), PortNo);
 	}
 #endif
-	
+
 	// create & initialize multicast socket (optional)
 	MulticastSocket = FUdpSocketBuilder(TEXT("UdpMessageMulticastSocket"))
 		.AsNonBlocking()
@@ -341,7 +343,6 @@ bool FUdpMessageTransport::StartTransport(IMessageTransportHandler& Handler)
 	MessageProcessor = new FUdpMessageProcessor(*MulticastSocket, FGuid::NewGuid(), MulticastEndpoint);
 #endif
 	{
-		UUdpMessagingSettings const* CurrentSettings = GetDefault<UUdpMessagingSettings>();
 		if (CurrentSettings->MaxConcurrentDeserializationTasks > 1)
 		{
 			UE_LOG(LogUdpMessaging, Display, TEXT("Using asynchronous task graph for message deserialization."));
@@ -367,7 +368,7 @@ bool FUdpMessageTransport::StartTransport(IMessageTransportHandler& Handler)
 	{
 		MulticastReceiver = new FUdpSocketReceiver(MulticastSocket, ThreadWaitTime, TEXT("UdpMessageMulticastReceiver"));
 		MulticastReceiver->OnDataReceived().BindRaw(this, &FUdpMessageTransport::HandleMulticastSocketDataReceived);
-		MulticastReceiver->SetMaxReadBufferSize(2048);
+		MulticastReceiver->SetMaxReadBufferSize(65535);
 		MulticastReceiver->Start();
 	}
 	else
@@ -379,7 +380,7 @@ bool FUdpMessageTransport::StartTransport(IMessageTransportHandler& Handler)
 	UnicastReceiver = new FUdpSocketReceiver(UnicastSocket, ThreadWaitTime, TEXT("UdpMessageUnicastReceiver"));
 	{
 		UnicastReceiver->OnDataReceived().BindRaw(this, &FUdpMessageTransport::HandleSocketDataReceived);
-		UnicastReceiver->SetMaxReadBufferSize(2048);
+		UnicastReceiver->SetMaxReadBufferSize(65535);
 		UnicastReceiver->Start();
 	}
 	MessageLogger = MakeUnique<FUdpMessageProcessorLogger>(UnicastEndpoint, UnicastSocket ? UnicastSocket->GetPortNo() : 0, MessageProcessor);
diff --git a/Engine/Plugins/Messaging/UdpMessaging/Source/UdpMessaging/Private/Transport/UdpReassembledMessage.h b/Engine/Plugins/Messaging/UdpMessaging/Source/UdpMessaging/Private/Transport/UdpReassembledMessage.h
--- a/Engine/Plugins/Messaging/UdpMessaging/Source/UdpMessaging/Private/Transport/UdpReassembledMessage.h
+++ b/Engine/Plugins/Messaging/UdpMessaging/Source/UdpMessaging/Private/Transport/UdpReassembledMessage.h
@@ -22,9 +22,6 @@ public:
 		PendingDelivery,
 		Delivered
 	};
-		
-	/** Default constructor. */
-	FUdpReassembledMessage() = default;
 
 	/**
 	 * Creates and initializes a new inbound message info.
@@ -207,14 +204,15 @@ public:
 	}
 
 	/**
-	 * Get the list of pending Acknowledgments up to the size of one segment size and clear them
+	 * Gets at most the specified number of pending acknowledgments and clears them.
 	 *
-	 * @return The array of pending acknowledgments removed
+	 * @return The array of pending acknowledgments removed.
 	 */
-	TArray<uint32> GetPendingAcknowledgments() 
+	TArray<uint32> TakePendingAcknowledgments(int32 MaxAckNum)
 	{
+		check(MaxAckNum > 0);
+
 		TArray<uint32> Temp;
-		const int32 MaxAckNum = (UDP_MESSAGING_SEGMENT_SIZE / sizeof(uint32));
 		if (PendingAcknowledgments.Num() < MaxAckNum)
 		{
 			Swap(Temp, PendingAcknowledgments);
diff --git a/Engine/Plugins/Messaging/UdpMessaging/Source/UdpMessaging/Private/Transport/UdpSerializeMessageTask.cpp b/Engine/Plugins/Messaging/UdpMessaging/Source/UdpMessaging/Private/Transport/UdpSerializeMessageTask.cpp
--- a/Engine/Plugins/Messaging/UdpMessaging/Source/UdpMessaging/Private/Transport/UdpSerializeMessageTask.cpp
+++ b/Engine/Plugins/Messaging/UdpMessaging/Source/UdpMessaging/Private/Transport/UdpSerializeMessageTask.cpp
@@ -10,6 +10,7 @@
 #include "UdpMessagingPrivate.h"
 #include "Transport/UdpSerializedMessage.h"
 #include "Shared/UdpMessagingSettings.h"
+#include "UdpMessageSegmenter.h"
 #include "UdpMessagingTracing.h"
 
 namespace UdpSerializeMessageTaskDetails
@@ -93,8 +94,8 @@ void SerializeMessageV11_15(FArchive& Archive, const TSharedRef<IMessageContext,
 	FStructSerializer::Serialize(MessageContext->GetMessage(), *MessageContext->GetMessageTypeInfo(), Backend);
 }
 
-/** Serialization Routine for message using Protocol version 16. */
-void SerializeMessageV16_17(FArchive& Archive, const TSharedRef<IMessageContext, ESPMode::ThreadSafe>& MessageContext, EUdpMessageFormat MessageFormat, const EStructSerializerBackendFlags StructSerializerBackendFlags)
+/** Serialization Routine for message using Protocol version 16+. */
+void SerializeMessageV16_18(FArchive& Archive, const TSharedRef<IMessageContext, ESPMode::ThreadSafe>& MessageContext, EUdpMessageFormat MessageFormat, const EStructSerializerBackendFlags StructSerializerBackendFlags)
 {
 	const FTopLevelAssetPath& MessageType = MessageContext->GetMessageTypePathName();
 	Archive << const_cast<FTopLevelAssetPath&>(MessageType);
@@ -127,7 +128,7 @@ void SerializeMessageV16_17(FArchive& Archive, const TSharedRef<IMessageContext,
 	}
 
 	// Message Wire Format Id
-	check(MessageFormat == EUdpMessageFormat::CborPlatformEndianness || MessageFormat == EUdpMessageFormat::CborStandardEndianness); // Versions 11 to 15 only supports CBOR.
+	check(MessageFormat == EUdpMessageFormat::CborPlatformEndianness || MessageFormat == EUdpMessageFormat::CborStandardEndianness); // Versions 11 and up only support CBOR.
 	uint8 Format = (uint8)(MessageFormat);
 	Archive << Format;
 
@@ -151,50 +152,69 @@ void FUdpSerializeMessageTask::DoTask(ENamedThreads::Type CurrentThread, const F
 		// a consistent wire format, if their implementations change. This allows us to sanity
 		// check the values during deserialization. @see FUdpDeserializeMessage::Deserialize()
 
+		const int64 MaxNumSegmentsV10_V11 = (int64)UINT16_MAX;
+		const int64 MaxNumSegments = (int64)INT32_MAX;
+
+		static_assert(UUdpMessagingSettings::MaxPacketSize_Min >
+			(UE::UdpMessaging::PacketHeaderBytes + FUdpMessageSegmenter::GetDataSegmentHeaderSize()),
+			"MaxPacketSize should always accommodate fixed headers");
+
+		const int32 MaxSegmentSize = (int32)GetDefault<UUdpMessagingSettings>()->GetMaxPacketSize() - UE::UdpMessaging::PacketHeaderBytes;
+		const uint16 MaxSegmentPayload = IntCastChecked<uint16>(MaxSegmentSize - FUdpMessageSegmenter::GetDataSegmentHeaderSize());
+
+		// 2048 corresponds to old UnicastReceiver MaxReadBufferSize
+		const int32 MaxSegmentSizeV10_V17 = FMath::Min(2048, MaxSegmentSize);
+		const uint16 MaxSegmentPayloadV10_V17 = IntCastChecked<uint16>(MaxSegmentSizeV10_V17 - FUdpMessageSegmenter::GetDataSegmentHeaderSize());
+
 		// serialize context depending on supported protocol version
-		int64 ProtocolMaxSegmentSize = UDP_MESSAGING_SEGMENT_SIZE * (int64)UINT16_MAX;
+		int64 ProtocolMaxMessageSize = MaxSegmentPayloadV10_V17 * MaxNumSegmentsV10_V11;
 		FArchive& Archive = SerializedMessage.Get();
 		bool Serialized = true;
 		switch (SerializedMessage->GetProtocolVersion())
 		{
 			case 10:
-				ProtocolMaxSegmentSize = UDP_MESSAGING_SEGMENT_SIZE * (int64)UINT16_MAX;
+				ProtocolMaxMessageSize = MaxSegmentPayloadV10_V17 * MaxNumSegmentsV10_V11;
 				UdpSerializeMessageTaskDetails::SerializeMessageV10(Archive, MessageContext);
 				break;
 
 			case 11:
-				ProtocolMaxSegmentSize = UDP_MESSAGING_SEGMENT_SIZE * (int64)UINT16_MAX;
+				ProtocolMaxMessageSize = MaxSegmentPayloadV10_V17 * MaxNumSegmentsV10_V11;
 				UdpSerializeMessageTaskDetails::SerializeMessageV11_15(Archive, MessageContext, EUdpMessageFormat::CborPlatformEndianness, EStructSerializerBackendFlags::None);
 				break;
 
 			case 12:
-				ProtocolMaxSegmentSize = UDP_MESSAGING_SEGMENT_SIZE * (int64)INT32_MAX;
+				ProtocolMaxMessageSize = MaxSegmentPayloadV10_V17 * MaxNumSegments;
 				UdpSerializeMessageTaskDetails::SerializeMessageV11_15(Archive, MessageContext, EUdpMessageFormat::CborPlatformEndianness, EStructSerializerBackendFlags::WriteTextAsComplexString);
 				break;
 
 			case 13:
-				ProtocolMaxSegmentSize = UDP_MESSAGING_SEGMENT_SIZE * (int64)INT32_MAX;
+				ProtocolMaxMessageSize = MaxSegmentPayloadV10_V17 * MaxNumSegments;
 				UdpSerializeMessageTaskDetails::SerializeMessageV11_15(Archive, MessageContext, EUdpMessageFormat::CborPlatformEndianness, EStructSerializerBackendFlags::LegacyUE4);
 				break;
 
 			case 14:
-				ProtocolMaxSegmentSize = UDP_MESSAGING_SEGMENT_SIZE * (int64)INT32_MAX;
+				ProtocolMaxMessageSize = MaxSegmentPayloadV10_V17 * MaxNumSegments;
 				UdpSerializeMessageTaskDetails::SerializeMessageV11_15(Archive, MessageContext, SerializedMessage->GetFormat(), EStructSerializerBackendFlags::LegacyUE4);
 				break;
 
 			case 15:
-				ProtocolMaxSegmentSize = UDP_MESSAGING_SEGMENT_SIZE * (int64)INT32_MAX;
+				ProtocolMaxMessageSize = MaxSegmentPayloadV10_V17 * MaxNumSegments;
 				UdpSerializeMessageTaskDetails::SerializeMessageV11_15(Archive, MessageContext, SerializedMessage->GetFormat(), EStructSerializerBackendFlags::Default);
 				break;
 
 			case 16:
-				ProtocolMaxSegmentSize = UDP_MESSAGING_SEGMENT_SIZE * (int64)INT32_MAX;
-				UdpSerializeMessageTaskDetails::SerializeMessageV16_17(Archive, MessageContext, SerializedMessage->GetFormat(), EStructSerializerBackendFlags::Default);
+				ProtocolMaxMessageSize = MaxSegmentPayloadV10_V17 * MaxNumSegments;
+				UdpSerializeMessageTaskDetails::SerializeMessageV16_18(Archive, MessageContext, SerializedMessage->GetFormat(), EStructSerializerBackendFlags::Default);
 				break;
 
 			case 17:
-				ProtocolMaxSegmentSize = UDP_MESSAGING_SEGMENT_SIZE * (int64)INT32_MAX;
-				UdpSerializeMessageTaskDetails::SerializeMessageV16_17(Archive, MessageContext, SerializedMessage->GetFormat(), EStructSerializerBackendFlags::Default);
+				ProtocolMaxMessageSize = MaxSegmentPayloadV10_V17 * MaxNumSegments;
+				UdpSerializeMessageTaskDetails::SerializeMessageV16_18(Archive, MessageContext, SerializedMessage->GetFormat(), EStructSerializerBackendFlags::Default);
+				break;
+
+			case 18:
+				ProtocolMaxMessageSize = MaxSegmentPayload * MaxNumSegments;
+				UdpSerializeMessageTaskDetails::SerializeMessageV16_18(Archive, MessageContext, SerializedMessage->GetFormat(), EStructSerializerBackendFlags::Default);
 				break;
 
 			default:
@@ -210,9 +230,9 @@ void FUdpSerializeMessageTask::DoTask(ENamedThreads::Type CurrentThread, const F
 			SerializedMessage->UpdateState(EUdpSerializedMessageState::Invalid);
 		}
 		// Once serialized if the size of the message is bigger than the maximum allow mark it as invalid and log an error
-		else if (SerializedMessage->TotalSize() > ProtocolMaxSegmentSize)
+		else if (SerializedMessage->TotalSize() > ProtocolMaxMessageSize)
 		{
-			UE_LOG(LogUdpMessaging, Error, TEXT("Serialized Message total size '%" INT64_FMT "' is over the allowed maximum '%" INT64_FMT "', discarding..."), SerializedMessage->TotalSize(), ProtocolMaxSegmentSize);
+			UE_LOG(LogUdpMessaging, Error, TEXT("Serialized Message total size '%" INT64_FMT "' is over the allowed maximum '%" INT64_FMT "', discarding..."), SerializedMessage->TotalSize(), ProtocolMaxMessageSize);
 			SerializedMessage->UpdateState(EUdpSerializedMessageState::Invalid);
 		}
 		else
diff --git a/Engine/Plugins/Messaging/UdpMessaging/Source/UdpMessaging/Private/UdpMessagingModule.cpp b/Engine/Plugins/Messaging/UdpMessaging/Source/UdpMessaging/Private/UdpMessagingModule.cpp
--- a/Engine/Plugins/Messaging/UdpMessaging/Source/UdpMessaging/Private/UdpMessagingModule.cpp
+++ b/Engine/Plugins/Messaging/UdpMessaging/Source/UdpMessaging/Private/UdpMessagingModule.cpp
@@ -629,7 +629,6 @@ protected:
 			FParse::Bool(CommandLine, TEXT("-UDPMESSAGING_SHARE_KNOWN_NODES="), Settings->bShareKnownNodesWithActiveConnections);
 			FParse::Value(CommandLine, TEXT("-UDPMESSAGING_TRANSPORT_UNICAST="), Settings->UnicastEndpoint);
 			FParse::Value(CommandLine, TEXT("-UDPMESSAGING_TRANSPORT_MULTICAST="), Settings->MulticastEndpoint);
-			FParse::Value(CommandLine, TEXT("-UDPMESSAGING_WORK_QUEUE_SIZE="), Settings->WorkQueueSize);
 
 			FString StaticEndpoints;
 			FParse::Value(CommandLine, TEXT("-UDPMESSAGING_TRANSPORT_STATIC="), StaticEndpoints, false);
@@ -639,7 +638,7 @@ protected:
 			{
 				Settings->StaticEndpoints.AddUnique(CmdStaticEndpoint);
 			}
-			UE_LOG(LogUdpMessaging, Display, TEXT("Work queue size set to %d."), Settings->WorkQueueSize);
+
 			if (Settings->bShareKnownNodesWithActiveConnections)
 			{
 				UE_LOG(LogUdpMessaging, Display, TEXT("Sharing known nodes with endpoints is enabled."));
diff --git a/Engine/Plugins/Messaging/UdpMessaging/Source/UdpMessaging/Private/UdpMessagingPrivate.h b/Engine/Plugins/Messaging/UdpMessaging/Source/UdpMessaging/Private/UdpMessagingPrivate.h
--- a/Engine/Plugins/Messaging/UdpMessaging/Source/UdpMessaging/Private/UdpMessagingPrivate.h
+++ b/Engine/Plugins/Messaging/UdpMessaging/Source/UdpMessaging/Private/UdpMessagingPrivate.h
@@ -21,11 +21,53 @@ DECLARE_LOG_CATEGORY_EXTERN(LogUdpMessaging, Log, All);
 /** Defines the desired size of socket receive buffers (in bytes). */
 #define UDP_MESSAGING_RECEIVE_BUFFER_SIZE 2 * 1024 * 1024
 
-/** Define the desired size for the message segments */
-#define UDP_MESSAGING_SEGMENT_SIZE 1024
-
 /**
  * Defines the protocol version of the UDP message transport.
  * @note When changing the version, ensure to update the serialization/deserialization code in UdpSerializeMessageTask.cpp/UdpDeserializedMessage.cpp and the supported versions in FUdpMessageProcessor::Init().
  */
-#define UDP_MESSAGING_TRANSPORT_PROTOCOL_VERSION 17
+#define UDP_MESSAGING_TRANSPORT_PROTOCOL_VERSION 18
+
+
+namespace UE::UdpMessaging
+{
+	/********************************
+	 *
+	 *   UDP Messaging Terminology
+	 *
+	 ********************************
+	 *
+	 *   Protocol Data Unit (PDU):
+	 *       Atomic unit of data, including any headers, payload, or trailers used by the protocol.
+	 *
+	 *   Service Data Unit (SDU):
+	 *       The usable payload component of the PDU. Populated by the next higher layer.
+	 *
+	 *   Maximum Transmission Unit (MTU):
+	 *       The largest allowable size for a PDU at a given layer.
+	 *
+	 *   Fragmentation:
+	 *       The IP layer transparently dividing and reassembling your oversized payload into multiple PDUs.
+	 *
+	 *   Segmentation:
+	 *       Trying to respect the lower layer's MTU and avoid fragmentation by proactively dividing payload.
+	 *       Can be (hardware) accelerated, such as by UDP segmentation offload (USO).
+	 *
+	 *   Frame:
+	 *       Ethernet / OSI layer 2 PDU. Typical MTU of 1500 bytes. "Jumbo frame" support allows MTU of ~9000 bytes.
+	 *
+	 *   Packet:
+	 *       IP / OSI layer 3 PDU. Usable MTU depends on routing path; ~1200 bytes is generally internet safe.
+	 *
+	 *   Datagram:
+	 *       UDP / OSI layer 4 PDU. Constitutes the SDU of the IP packet.
+	 * 
+	 *   Segment:
+	 *       UDP Messaging PDU (You are here). Encapsulated within (SDU of) a UDP datagram.
+	 * 
+	 */
+
+	// Technically minimum, but almost always holds in practice.
+	constexpr uint16 Ipv4HeaderBytes = 20;
+	constexpr uint16 UdpHeaderBytes = 8;
+	constexpr uint16 PacketHeaderBytes = Ipv4HeaderBytes + UdpHeaderBytes;
+}
diff --git a/Engine/Plugins/Messaging/UdpMessaging/Source/UdpMessaging/Public/Shared/UdpMessagingSettings.h b/Engine/Plugins/Messaging/UdpMessaging/Source/UdpMessaging/Public/Shared/UdpMessagingSettings.h
--- a/Engine/Plugins/Messaging/UdpMessaging/Source/UdpMessaging/Public/Shared/UdpMessagingSettings.h
+++ b/Engine/Plugins/Messaging/UdpMessaging/Source/UdpMessaging/Public/Shared/UdpMessagingSettings.h
@@ -79,15 +79,33 @@ public:
 	UPROPERTY(config, EditAnywhere, Category = Transport, AdvancedDisplay)
 	uint32 AutoRepairAttemptLimit = 10;
 
+	UE_DEPRECATED(5.8, "UDP Messaging work queues are no longer fixed size.")
+	UPROPERTY(meta=(DeprecatedProperty))
+	uint16 WorkQueueSize = 1024;
+
 	/**
-	 * The buffer size for the working send queues. Each node connections
-	 * gets a send queue and new messages to send are put on that queue.  The send
-	 * scheduler will fairly send data on this queue and re-queue when partial data is
-	 * sent or waiting for a reliable message. Can be specified on the command line with
-	 * `-UDPMESSAGING_WORK_QUEUE_SIZE=`
+	 * The maximum packet size to use for message segments (including presumed 20 byte IPv4 + 8 byte UDP headers).
+	 * This should be less than or equal to your path maximum transmission unit (PMTU) to avoid fragmentation.
+	 * Previous versions of UDP Messaging were hard-coded to 1130, which is a conservative, internet-safe value.
+	 *
+	 * If your traffic is restricted to the local network, support for an MTU of 1500 is common.
+	 * This can be increased further by enabling jumbo frames on supported hardware.
 	 */
-	UPROPERTY(config, EditAnywhere, Category = Transport, meta = (ClampMin="1024"), AdvancedDisplay)
-	uint16 WorkQueueSize = 1024;
+	UPROPERTY(config, EditAnywhere, Category = Transport, meta = (ClampMin="576", ClampMax="65535"), AdvancedDisplay)
+	uint16 MaxPacketSize = 1130;
+
+	static constexpr uint16 MaxPacketSize_Min = 576;
+	static constexpr uint16 MaxPacketSize_Max = 65535;
+
+	/** Accessor that ensures MaxPacketSize is clamped to the expected bounds. */
+	uint16 GetMaxPacketSize() const
+	{
+		return FMath::Clamp<uint16>(MaxPacketSize, MaxPacketSize_Min, MaxPacketSize_Max);
+	}
+
+	/** Enable UDP segmentation offload. */
+	UPROPERTY(config, EditAnywhere, Category = Transport, AdvancedDisplay)
+	bool bSegmentationOffload = false;
 
 	/**
 	 * Reliable queue priority specified as a percentage with a range from 5-95%. Setting values above or below
diff --git a/Engine/Source/Runtime/Sockets/Private/BSDSockets/SocketsBSD.cpp b/Engine/Source/Runtime/Sockets/Private/BSDSockets/SocketsBSD.cpp
--- a/Engine/Source/Runtime/Sockets/Private/BSDSockets/SocketsBSD.cpp
+++ b/Engine/Source/Runtime/Sockets/Private/BSDSockets/SocketsBSD.cpp
@@ -646,6 +646,16 @@ bool FSocketBSD::SetReceiveBufferSize(int32 Size,int32& NewSize)
 	return bOk;
 }
 
+bool FSocketBSD::SetSendSegmentationOffloadSize(uint16 InSegmentPayloadBytes)
+{
+#if PLATFORM_HAS_BSD_SOCKET_FEATURE_WINSOCKETS
+	uint32 OptSegmentSize = InSegmentPayloadBytes;
+	return 0 == setsockopt(Socket, IPPROTO_UDP, UDP_SEND_MSG_SIZE, (char*)&OptSegmentSize, sizeof(OptSegmentSize));
+#else
+	return false;
+#endif
+}
+
 int32 FSocketBSD::GetPortNo(void)
 {
 	sockaddr_storage Addr;
diff --git a/Engine/Source/Runtime/Sockets/Private/BSDSockets/SocketsBSD.h b/Engine/Source/Runtime/Sockets/Private/BSDSockets/SocketsBSD.h
--- a/Engine/Source/Runtime/Sockets/Private/BSDSockets/SocketsBSD.h
+++ b/Engine/Source/Runtime/Sockets/Private/BSDSockets/SocketsBSD.h
@@ -118,6 +118,7 @@ public:
 	virtual bool SetRecvErr(bool bUseErrorQueue = true) override;
 	virtual bool SetSendBufferSize(int32 Size,int32& NewSize) override;
 	virtual bool SetReceiveBufferSize(int32 Size,int32& NewSize) override;
+	virtual bool SetSendSegmentationOffloadSize(uint16 InSegmentPayloadBytes) override;
 	virtual int32 GetPortNo() override;
 	virtual bool SetIpPktInfo(bool bEnable) override;
 	virtual bool RecvFromWithPktInfo(uint8* Data, int32 BufferSize, int32& BytesRead, FInternetAddr& Source, FInternetAddr& Destination, ESocketReceiveFlags::Type Flags = ESocketReceiveFlags::None) override;
diff --git a/Engine/Source/Runtime/Sockets/Private/Sockets.cpp b/Engine/Source/Runtime/Sockets/Private/Sockets.cpp
--- a/Engine/Source/Runtime/Sockets/Private/Sockets.cpp
+++ b/Engine/Source/Runtime/Sockets/Private/Sockets.cpp
@@ -89,6 +89,11 @@ bool FSocket::SetRetrieveTimestamp(bool bRetrieveTimestamp/*=true*/)
 	return false;
 }
 
+bool FSocket::SetSendSegmentationOffloadSize(uint16 InSegmentPayloadBytes)
+{
+	return false;
+}
+
 bool FSocket::SetIpPktInfo(bool bEnable)
 {
 	return false;
diff --git a/Engine/Source/Runtime/Sockets/Public/Sockets.h b/Engine/Source/Runtime/Sockets/Public/Sockets.h
--- a/Engine/Source/Runtime/Sockets/Public/Sockets.h
+++ b/Engine/Source/Runtime/Sockets/Public/Sockets.h
@@ -416,6 +416,16 @@ public:
 	 */
 	SOCKETS_API virtual bool SetRetrieveTimestamp(bool bRetrieveTimestamp=true);
 
+	/**
+	 * Attempt to enable UDP segmentation offload (USO). Where supported, this
+	 * reduces overhead by automatically splitting "superpackets" which exceed
+	 * the specified segment size into multiple datagrams.
+	 *
+	 * @param InSegmentPayloadBytes Number of payload bytes for each split segment.
+	 * @return True if the call succeeded, false otherwise.
+	 */
+	SOCKETS_API virtual bool SetSendSegmentationOffloadSize(uint16 InSegmentPayloadBytes);
+
 	/**
 	 * Reads the port this socket is bound to.
 	 *
