UdpMessaging: Use high resolution monotonic timer, improved segment retransmit logic. (WIP)

- Introduce FUdpMessagingTime as a thin wrapper around FPlatformTime, and use it instead of FDateTime.
- Take up-to-date timings throughout the processor, rather than caching a single "CurrentTime" ~once per thread wakeup.
- Adopt RFC6298-inspired retransmission timeout calculation, which uses a more robust estimate of RTT, and incorporates the considerable variance we see in practice.
- Optimize send window segment expiry, and run it unconditionally on both reliable and unreliable segments.
- Remove redundant retransmit logic that was internal to the segmenter now that we ensure timely expiry of all sends.
- Coalesce segment acknowledgments until we have a full MSS worth, or until 5ms have elapsed (in the spirit of Nagle's algorithm).
- Limit `ConsumeInboundSegments` to ~5ms time slices rather than (0.25 * 5 * 1000) = ~1250+ms. In practice this seems to have no impact on throughput while improving tail latencies.
- Introduce a separate "high-priority" queue in FSocketSender, currently only used for reducing latency on acknowledgments.
- Add Insights trace scopes to more code.

diff --git a/Engine/Plugins/Messaging/UdpMessaging/Source/UdpMessaging/Private/Shared/SocketSender.cpp b/Engine/Plugins/Messaging/UdpMessaging/Source/UdpMessaging/Private/Shared/SocketSender.cpp
--- a/Engine/Plugins/Messaging/UdpMessaging/Source/UdpMessaging/Private/Shared/SocketSender.cpp
+++ b/Engine/Plugins/Messaging/UdpMessaging/Source/UdpMessaging/Private/Shared/SocketSender.cpp
@@ -9,6 +9,7 @@
 #include "Sockets.h"
 #include "SocketSubsystem.h"
 #include "UdpMessagingPrivate.h"
+#include "UdpMessagingTracing.h"
 
 
 namespace UE::UdpMessaging {
@@ -58,7 +59,8 @@ FSocketSender::~FSocketSender()
 bool FSocketSender::Send(
 	const TSharedRef<TArray<uint8>>& InData,
 	const FIPv4Endpoint& InRecipient,
-	bool bAllowSegmentation /* = false */
+	bool bAllowSegmentation, /* = false */
+    bool bHighPriority /* = false */
 )
 {
 	if (!bStopping)
@@ -71,7 +73,8 @@ bool FSocketSender::Send(
 				TEXT("Send data exceeds USO offload size, but bAllowSegmentation was false"));
 		}
 
-		SendQueue.Enqueue(InData, InRecipient);
+		TMpscQueue<FPacket>& Queue = bHighPriority ? PrioritySendQueue : SendQueue;
+		Queue.Enqueue(InData, InRecipient);
 		WorkEvent->Trigger();
 		return true;
 	}
@@ -110,14 +113,18 @@ void FSocketSender::Stop()
 
 FSocketSender::EUpdateResult FSocketSender::Update(const FTimespan& InSocketWaitTime)
 {
-	while (!SendQueue.IsEmpty())
+	while (!PrioritySendQueue.IsEmpty() || !SendQueue.IsEmpty())
 	{
 		if (!Socket->Wait(ESocketWaitConditions::WaitForWrite, InSocketWaitTime))
 		{
 			return EUpdateResult::Retry;
 		}
 
-		const FPacket* Packet = SendQueue.Peek();
+		SCOPED_MESSAGING_TRACE(UdpMessaging_FSocketSender_Update);
+
+		TMpscQueue<FPacket>& Queue = !PrioritySendQueue.IsEmpty() ? PrioritySendQueue : SendQueue;
+
+		const FPacket* Packet = Queue.Peek();
 		const FIPv4Endpoint Recipient = Packet->Recipient;
 		const int32 PacketNumBytes = Packet->Data->Num();
 		int32 SentNumBytes = 0;
@@ -138,7 +145,7 @@ FSocketSender::EUpdateResult FSocketSender::Update(const FTimespan& InSocketWait
 
 		// Only pop after we've checked for "retry-able" error codes.
 		Packet = nullptr;
-		SendQueue.Dequeue();
+		Queue.Dequeue();
 
 		if (!bSendResult)
 		{
diff --git a/Engine/Plugins/Messaging/UdpMessaging/Source/UdpMessaging/Private/Shared/SocketSender.h b/Engine/Plugins/Messaging/UdpMessaging/Source/UdpMessaging/Private/Shared/SocketSender.h
--- a/Engine/Plugins/Messaging/UdpMessaging/Source/UdpMessaging/Private/Shared/SocketSender.h
+++ b/Engine/Plugins/Messaging/UdpMessaging/Source/UdpMessaging/Private/Shared/SocketSender.h
@@ -91,12 +91,14 @@ public:
 	 * @param InData The data to send.
 	 * @param InRecipient The recipient.
 	 * @param bAllowSegmentation Pass true if you are intentionally sending more than `GetSegmentOffloadSize()` bytes.
+	 * @param bHighPriority Packets in this queue are always sent first.
 	 * @return true if the data will be sent, false otherwise.
 	 */
 	bool Send(
 		const TSharedRef<TArray<uint8>>& InData,
 		const FIPv4Endpoint& InRecipient,
-		bool bAllowSegmentation = false
+		bool bAllowSegmentation = false,
+		bool bHighPriority = false
 	);
 
 	/** If send segmentation offload is enabled, returns the max datagram payload. Otherwise 0. */
@@ -137,6 +139,9 @@ private:
 	/** The send queue. */
 	TMpscQueue<FPacket> SendQueue;
 
+	/** High-priority send queue. */
+	TMpscQueue<FPacket> PrioritySendQueue;
+
 	/** The network socket. */
 	FSocket* Socket;
 
diff --git a/Engine/Plugins/Messaging/UdpMessaging/Source/UdpMessaging/Private/Shared/UdpMessagingTime.h b/Engine/Plugins/Messaging/UdpMessaging/Source/UdpMessaging/Private/Shared/UdpMessagingTime.h
new file mode 100644
--- /dev/null
+++ b/Engine/Plugins/Messaging/UdpMessaging/Source/UdpMessaging/Private/Shared/UdpMessagingTime.h
@@ -0,0 +1,48 @@
+// Copyright Epic Games, Inc. All Rights Reserved.
+
+#pragma once
+
+#include "HAL/PlatformTime.h"
+#include "Misc/Timespan.h"
+
+
+/** Thin wrapper to help migrate from FDateTime to FPlatformTime::Seconds for relative measurements. */
+struct FUdpMessagingTime
+{
+	double RawValue = std::numeric_limits<double>::lowest();
+
+	static FUdpMessagingTime Now()
+	{
+		return { .RawValue = FPlatformTime::Seconds() };
+	}
+
+	FTimespan operator-(FUdpMessagingTime Other) const
+	{
+		return FTimespan::FromSeconds(RawValue - Other.RawValue);
+	}
+
+	auto operator+(FUdpMessagingTime Other) const = delete;
+	auto operator+=(FUdpMessagingTime Other) const = delete;
+
+	FUdpMessagingTime operator-(FTimespan Other) const
+	{
+		return { .RawValue = RawValue - Other.GetTotalSeconds() };
+	}
+
+	FUdpMessagingTime operator+(FTimespan Other) const
+	{
+		return { .RawValue = RawValue + Other.GetTotalSeconds() };
+	}
+
+#if 0 // Enable some day when all our target toolchains support it.
+	/** Consistent compiler generated default comparison operations based on RawValue. */
+	friend auto operator<=>(FUdpMessagingTime, FUdpMessagingTime) = default;
+#else
+	bool operator==(FUdpMessagingTime Other) const { return RawValue == Other.RawValue; }
+	bool operator!=(FUdpMessagingTime Other) const { return RawValue != Other.RawValue; }
+	bool operator>(FUdpMessagingTime Other) const { return RawValue > Other.RawValue; }
+	bool operator>=(FUdpMessagingTime Other) const { return RawValue >= Other.RawValue; }
+	bool operator<(FUdpMessagingTime Other) const { return RawValue < Other.RawValue; }
+	bool operator<=(FUdpMessagingTime Other) const { return RawValue <= Other.RawValue; }
+#endif
+};
diff --git a/Engine/Plugins/Messaging/UdpMessaging/Source/UdpMessaging/Private/Transport/UdpMessageProcessor.cpp b/Engine/Plugins/Messaging/UdpMessaging/Source/UdpMessaging/Private/Transport/UdpMessageProcessor.cpp
--- a/Engine/Plugins/Messaging/UdpMessaging/Source/UdpMessaging/Private/Transport/UdpMessageProcessor.cpp
+++ b/Engine/Plugins/Messaging/UdpMessaging/Source/UdpMessaging/Private/Transport/UdpMessageProcessor.cpp
@@ -11,7 +11,6 @@
 
 #include "Transport/UdpDeserializedMessage.h"
 #include "UdpMessagingPrivate.h"
-#include "UdpMessagingTracing.h"
 #include "Shared/SocketSender.h"
 #include "Shared/UdpMessagingSettings.h"
 #include "Transport/UdpMessageBeacon.h"
@@ -45,12 +44,6 @@ TAutoConsoleVariable<FString> CVarConnectionsToError(
 	TEXT("This can be a comma separated list in the form IPAddr2:port,IPAddr3:port"),
 	ECVF_Default);
 
-TAutoConsoleVariable<int32> CVarCheckForExpired(
-	TEXT("MessageBus.UDP.CheckForExpired"),
-	1,
-	TEXT("Attempts to relieve pressure on the work queue by checking if inflight segments have expired with no acknowledgement.\n"),
-	ECVF_Default);
-
 TAutoConsoleVariable<int32> CVarUdpMessagingMaxWindowSize(
 	TEXT("MessageBus.UDP.MaxWindowSize"),
 	2048,
@@ -109,8 +102,7 @@ FOnInboundTransferDataUpdated& OnReassemblerUpdated()
 /* FUdpMessageProcessor structors
  *****************************************************************************/
 FUdpMessageProcessor::FNodeInfo::FNodeInfo()
-	: LastSegmentReceivedTime(FDateTime::MinValue())
-	, NodeId()
+	: NodeId()
 	, ProtocolVersion(UDP_MESSAGING_TRANSPORT_PROTOCOL_VERSION)
 	, ReliableWorkQueue(1024)
 	, UnreliableWorkQueue(1024)
@@ -138,7 +130,6 @@ FUdpMessageProcessor::FUdpMessageProcessor(FSocket& InSocket, const FGuid& InNod
 		FPlatformProcess::ReturnSynchEventToPool(EventToDelete);
 	});
 
-	CurrentTime = FDateTime::UtcNow();
 	const ELoadingPhase::Type LoadingPhase = IPluginManager::Get().GetLastCompletedLoadingPhase();
 	if (LoadingPhase == ELoadingPhase::None || LoadingPhase < ELoadingPhase::PostDefault)
 	{
@@ -384,15 +375,21 @@ bool FUdpMessageProcessor::Init()
 
 uint32 FUdpMessageProcessor::Run()
 {
+	FUdpMessagingTime LastTime = FUdpMessagingTime::Now();
+
 	while (!bStopping)
 	{
 		WorkEvent->Wait(CalculateWaitTime());
 
+		SCOPED_MESSAGING_TRACE(FUdpMessageProcessor_Run);
+
 		do
 		{
-			FDateTime LastTime = CurrentTime;
-			CurrentTime = FDateTime::UtcNow();
-			DeltaTime = CurrentTime - LastTime;
+			SCOPED_MESSAGING_TRACE(FUdpMessageProcessor_Run_Frame);
+
+			const FUdpMessagingTime Now = FUdpMessagingTime::Now();
+			DeltaTime = Now - LastTime;
+			LastTime = Now;
 
 			for (TSet<FUdpReassembledMessageSharedPtr>::TIterator It = PendingTypeLookup.CreateIterator(); It; ++It)
 			{
@@ -461,8 +458,6 @@ void FUdpMessageProcessor::WaitAsyncTaskCompletion()
 
 void FUdpMessageProcessor::Tick()
 {
-	CurrentTime = FDateTime::UtcNow();
-
 	ConsumeInboundSegments();
 	ConsumeOutboundMessages();
 	UpdateKnownNodes();
@@ -491,7 +486,9 @@ void FUdpMessageProcessor::AcknowledgeReceipt(int32 MessageId, const FNodeInfo&
 	*Writer << Header;
 	AcknowledgeChunk.Serialize(*Writer, NodeInfo.ProtocolVersion);
 
-	if (!SocketSender->Send(Writer, NodeInfo.Endpoint))
+	const bool bAllowSegmentation_false = false;
+	const bool bHighPriority_true = true;
+	if (!SocketSender->Send(Writer, NodeInfo.Endpoint, bAllowSegmentation_false, bHighPriority_true))
 	{
 		UE_LOG(LogUdpMessaging, Warning, TEXT("FUdpMessageProcessor::AcknowledgeReceipt failed to send %s."), *NodeInfo.Endpoint.ToString());
 		return;
@@ -503,22 +500,19 @@ void FUdpMessageProcessor::AcknowledgeReceipt(int32 MessageId, const FNodeInfo&
 
 FTimespan FUdpMessageProcessor::CalculateWaitTime() const
 {
+	// TODO: Peek at expiry queues to determine when it would be appropriate to wake sooner; potentially increase this?
 	return FTimespan::FromMilliseconds(10);
 }
 
 
-FTimespan ComputeMaxWorkTimespan(const int32 DeadHelloIntervals, FUdpMessageBeacon* Beacon)
-{
-	return 0.25 * DeadHelloIntervals * Beacon->GetBeaconInterval();
-}
-
 void FUdpMessageProcessor::ConsumeInboundSegments()
 {
 	SCOPED_MESSAGING_TRACE(FUdpMessageProcessor_ConsumeInboundSegments);
-	FInboundSegment Segment;
 
-	FTimespan MaxWorkTimespan = ComputeMaxWorkTimespan(DeadHelloIntervals, Beacon.Get());
+	const FTimespan MaxWorkTimespan = FTimespan::FromMilliseconds(5);
+	const FUdpMessagingTime EndWorkTime = FUdpMessagingTime::Now() + MaxWorkTimespan;
 
+	FInboundSegment Segment;
 	while (InboundSegments.Dequeue(Segment))
 	{
 		// quick hack for TTP# 247103
@@ -527,7 +521,6 @@ void FUdpMessageProcessor::ConsumeInboundSegments()
 			continue;
 		}
 
-		
 		FUdpMessageSegment::FHeader Header;
 		*Segment.Data << Header;
 
@@ -565,7 +558,7 @@ void FUdpMessageProcessor::ConsumeInboundSegments()
 			}
 
 			NodeInfo.Endpoint = Segment.Sender;
-			NodeInfo.LastSegmentReceivedTime = CurrentTime;
+			NodeInfo.LastSegmentReceivedTime = FUdpMessagingTime::Now();
 			NodeInfo.MaxWindowSize = CVarUdpMessagingMaxWindowSize.GetValueOnAnyThread();
 			
 			switch (Header.SegmentType)
@@ -619,7 +612,7 @@ void FUdpMessageProcessor::ConsumeInboundSegments()
 			}
 		}
 
-		if ((CurrentTime + MaxWorkTimespan) <= FDateTime::UtcNow())
+		if (FUdpMessagingTime::Now() >= EndWorkTime)
 		{
 			break;
 		}
@@ -717,6 +710,8 @@ bool FUdpMessageProcessor::FilterSegment(const FUdpMessageSegment::FHeader& Head
 
 void FUdpMessageProcessor::ProcessAbortSegment(FInboundSegment& Segment, FNodeInfo& NodeInfo)
 {
+	SCOPED_MESSAGING_TRACE(FUdpMessageProcessor_ProcessAbortSegment);
+
 	FUdpMessageSegment::FAbortChunk AbortChunk;
 	AbortChunk.Serialize(*Segment.Data, NodeInfo.ProtocolVersion);
 
@@ -726,6 +721,8 @@ void FUdpMessageProcessor::ProcessAbortSegment(FInboundSegment& Segment, FNodeIn
 
 void FUdpMessageProcessor::ProcessAcknowledgeSegment(FInboundSegment& Segment, FNodeInfo& NodeInfo)
 {
+	SCOPED_MESSAGING_TRACE(FUdpMessageProcessor_ProcessAcknowledgeSegment);
+
 	FUdpMessageSegment::FAcknowledgeChunk AcknowledgeChunk;
 	AcknowledgeChunk.Serialize(*Segment.Data, NodeInfo.ProtocolVersion);
 
@@ -735,7 +732,7 @@ void FUdpMessageProcessor::ProcessAcknowledgeSegment(FInboundSegment& Segment, F
 		TSharedPtr<FUdpMessageSegmenter>& Segmenter = *FoundSegmenter;
 		if (EnumHasAnyFlags(Segmenter->GetMessageFlags(), EMessageFlags::Reliable))
 		{
-			NodeInfo.MarkComplete(AcknowledgeChunk.MessageId, CurrentTime);
+			NodeInfo.MarkComplete(AcknowledgeChunk.MessageId, FUdpMessagingTime::Now());
 		}
 		SendSegmenterStatsToListeners(AcknowledgeChunk.MessageId, NodeInfo.NodeId, Segmenter);
 	}
@@ -748,16 +745,20 @@ void FUdpMessageProcessor::ProcessAcknowledgeSegment(FInboundSegment& Segment, F
 
 void FUdpMessageProcessor::ProcessAcknowledgeSegmentsSegment(FInboundSegment& Segment, FNodeInfo& NodeInfo) // TODO: Rename function
 {
+	SCOPED_MESSAGING_TRACE(FUdpMessageProcessor_ProcessAcknowledgeSegmentsSegment);
+
 	FUdpMessageSegment::FAcknowledgeSegmentsChunk AcknowledgeChunk;
 	AcknowledgeChunk.Serialize(*Segment.Data, NodeInfo.ProtocolVersion);
 
 	UE_LOG(LogUdpMessaging, Verbose, TEXT("Received AcknowledgeSegments for %d from %s"), AcknowledgeChunk.MessageId, *NodeInfo.NodeId.ToString());
-	NodeInfo.MarkAcks(AcknowledgeChunk.MessageId, AcknowledgeChunk.Segments, CurrentTime);
+	NodeInfo.MarkAcks(AcknowledgeChunk.MessageId, AcknowledgeChunk.Segments, FUdpMessagingTime::Now());
 }
 
 
 void FUdpMessageProcessor::ProcessByeSegment(FInboundSegment& Segment, FNodeInfo& NodeInfo)
 {
+	SCOPED_MESSAGING_TRACE(FUdpMessageProcessor_ProcessByeSegment);
+
 	FGuid RemoteNodeId;
 	*Segment.Data << RemoteNodeId;
 
@@ -770,6 +771,8 @@ void FUdpMessageProcessor::ProcessByeSegment(FInboundSegment& Segment, FNodeInfo
 
 void FUdpMessageProcessor::ProcessDataSegment(FInboundSegment& Segment, FNodeInfo& NodeInfo)
 {
+	SCOPED_MESSAGING_TRACE(FUdpMessageProcessor_ProcessDataSegment);
+
 	FUdpMessageSegment::FDataChunk DataChunk;
 	DataChunk.Serialize(*Segment.Data, NodeInfo.ProtocolVersion);
 
@@ -800,7 +803,7 @@ void FUdpMessageProcessor::ProcessDataSegment(FInboundSegment& Segment, FNodeInf
 
 	NodeInfo.Statistics.TotalBytesReceived += DataChunk.Data.Num();
 	NodeInfo.Statistics.PacketsReceived++;
-	ReassembledMessage->Reassemble(DataChunk.SegmentNumber, DataChunk.SegmentOffset, DataChunk.Data, CurrentTime);
+	ReassembledMessage->Reassemble(DataChunk.SegmentNumber, DataChunk.SegmentOffset, DataChunk.Data, FUdpMessagingTime::Now());
 
 	if (DataChunk.SegmentNumber == 0)
 	{
@@ -857,6 +860,8 @@ void FUdpMessageProcessor::ProcessHelloSegment(FInboundSegment& Segment, FNodeIn
 
 void FUdpMessageProcessor::ProcessPingSegment(FInboundSegment& Segment, FNodeInfo& NodeInfo)
 {
+	SCOPED_MESSAGING_TRACE(FUdpMessageProcessor_ProcessPingSegment);
+
 	FGuid RemoteNodeId;
 	*Segment.Data << RemoteNodeId;
 	uint8 NodeProtocolVersion;
@@ -916,6 +921,8 @@ void FUdpMessageProcessor::ProcessPongSegment(FInboundSegment& Segment, FNodeInf
 
 void FUdpMessageProcessor::ProcessRetransmitSegment(FInboundSegment& Segment, FNodeInfo& NodeInfo)
 {
+	SCOPED_MESSAGING_TRACE(FUdpMessageProcessor_ProcessRetransmitSegment);
+
 	FUdpMessageSegment::FRetransmitChunk RetransmitChunk;
 	int32 TargetMessageId = RetransmitChunk.GetMessageId(*Segment.Data);
 	if (NodeInfo.Segmenters.IsEmpty() || !NodeInfo.Segmenters.Contains(TargetMessageId))
@@ -943,6 +950,8 @@ void FUdpMessageProcessor::ProcessRetransmitSegment(FInboundSegment& Segment, FN
 
 void FUdpMessageProcessor::ProcessTimeoutSegment(FInboundSegment& Segment, FNodeInfo& NodeInfo)
 {
+	SCOPED_MESSAGING_TRACE(FUdpMessageProcessor_ProcessTimeoutSegment);
+
 	if (NodeInfo.Segmenters.IsEmpty())
 	{
 		// Ignore this message because we have no segmenters to retransmit.
@@ -1028,10 +1037,10 @@ int32 GetMaxSendRate()
 	return static_cast< uint32 > ( GetDefault<UUdpMessagingSettings>()->MaxSendRate * OneGbitPerSecondInBytes );
 }
 
-bool HasTimedOut(const FDateTime& ReceivedTime, const FDateTime& CurrentTime)
+bool HasTimedOut(const FUdpMessagingTime& ReceivedTime, const FUdpMessagingTime& CurrentTime)
 {
 	const float TimeoutInSeconds = GetDefault<UUdpMessagingSettings>()->ConnectionTimeoutPeriod;
-	FTimespan TimeoutInterval = TimeoutInSeconds * FTimespan::FromMilliseconds(1000);
+	FTimespan TimeoutInterval = FTimespan::FromSeconds(TimeoutInSeconds);
 	return (ReceivedTime + TimeoutInterval) < CurrentTime;
 }
 
@@ -1050,7 +1059,7 @@ void FUdpMessageProcessor::RemoveDeadNodes()
 		FGuid& NodeId = It->Key;
 		FNodeInfo& NodeInfo = It->Value;
 
-		if ((NodeId.IsValid()) && HasTimedOut(NodeInfo.LastSegmentReceivedTime, CurrentTime))
+		if ((NodeId.IsValid()) && HasTimedOut(NodeInfo.LastSegmentReceivedTime, FUdpMessagingTime::Now()))
 		{
 			UE_LOG(LogUdpMessaging, Display, TEXT("FUdpMessageProcessor::UpdateKnownNodes: Removing Node %s (%s)"), *NodeInfo.NodeId.ToString(), *NodeInfo.Endpoint.ToString());
 			NodeLostDelegate.ExecuteIfBound(NodeId);
@@ -1063,7 +1072,7 @@ bool FUdpMessageProcessor::MoreToSend()
 {
 	for (auto& KnownNodePair : KnownNodes)
 	{
-		if (KnownNodePair.Value.CanSendSegments() && KnownNodePair.Value.HasSegmenterThatCanSend(CurrentTime))
+		if (KnownNodePair.Value.CanSendSegments() && KnownNodePair.Value.HasSegmenterThatCanSend(FUdpMessagingTime::Now()))
 		{
 			return true;
 		}
@@ -1213,8 +1222,6 @@ void FUdpMessageProcessor::UpdateKnownNodes()
 			break;
 		}
 
-		KnownNodePair.Value.RemoveOldUnreliables(CurrentTime);
-
 		if (MemoryLoggingThresholdMB >= 0 && bCanLogDebugInfo)
 		{
 			uint32 ReassembledTotalBytes = 0;
@@ -1337,7 +1344,7 @@ FSentSegmentInfo FUdpMessageProcessor::SendNextSegmentsForMessageId(
 
 	FSentSegmentInfo SentInfo(MessageId);
 	TArray<FSentData, TInlineAllocator<64>> SentDatas;
-	if (Segmenter->IsInitialized() && Segmenter->NeedSending(CurrentTime))
+	if (Segmenter->IsInitialized() && !Segmenter->IsSendingComplete())
 	{
 		if (Header.ProtocolVersion != Segmenter->GetProtocolVersion())
 		{
@@ -1392,16 +1399,13 @@ FSentSegmentInfo FUdpMessageProcessor::SendNextSegmentsForMessageId(
 				DataChunk.SegmentNumber,
 				++NodeInfo.SequenceId,
 				bIsReliable,
-				CurrentTime
+				FUdpMessagingTime::Now()
 			);
 		}
 
 		const bool bAllowSegmentation_true = true;
 		if (SocketSender->Send(Writer, NodeInfo.Endpoint, bAllowSegmentation_true))
 		{
-			// FIXME?: NeedSending above implicitly updates the LastSentTime, even if send fails
-			Segmenter->UpdateSentTime(CurrentTime);
-
 			SentInfo.SequenceNumber = NodeInfo.SequenceId;
 			SentInfo.BytesSent = Writer->Tell();
 			SentInfo.SegmentsSent = SentDatas.Num();
@@ -1414,9 +1418,7 @@ FSentSegmentInfo FUdpMessageProcessor::SendNextSegmentsForMessageId(
 
 			for (FSentData& SentData : SentDatas)
 			{
-				NodeInfo.InflightSegments.Add(
-					NodeInfo.Remap(MessageId, SentData.SegmentNumber),
-					SentData);
+				NodeInfo.AddInflight(SentData);
 			}
 
 			SendSegmenterStatsToListeners(MessageId, NodeInfo.NodeId, Segmenter);
@@ -1433,6 +1435,7 @@ FSentSegmentInfo FUdpMessageProcessor::SendNextSegmentsForMessageId(
 			for (const FSentData& SentData : SentDatas)
 			{
 				// Roll back the sent state of any segments we failed to send.
+				NodeInfo.RemoveInflight(SentData);
 				Segmenter->MarkForRetransmission(SentData.SegmentNumber);
 			}
 
@@ -1581,13 +1584,11 @@ int32 FUdpMessageProcessor::UpdateSegmenters(FNodeInfo& NodeInfo)
 			return -1;
 		}
 	}
-	
-	if (CVarCheckForExpired.GetValueOnAnyThread() > 0 && !NodeInfo.CanSendSegments())
-	{
-		// RemoveLostSegments will clean-up the tracking map that we use to determine what is "inflight".
-		// Old data will be expired and scheduled for resend.
-		NodeInfo.RemoveLostSegments(CurrentTime);
-	}
+
+	// RemoveLostSegments will clean-up the tracking map that we use to determine what is "inflight".
+	// Old data will be expired and scheduled for resend.
+	const uint32 LostSegmentCount = NodeInfo.ExpireSegments(FUdpMessagingTime::Now());
+	NodeInfo.ComputeWindowSize(0, LostSegmentCount);
 
 	NodeInfo.Statistics.PacketsInFlight = NodeInfo.InflightSegments.Num();
 	NodeInfo.Statistics.BytesInflight = NodeInfo.InflightSegments.Num() * NodeInfo.MaxSegmentSize;
@@ -1600,6 +1601,7 @@ const FTimespan FUdpMessageProcessor::StaleReassemblyInterval = FTimespan::FromS
 bool FUdpMessageProcessor::UpdateReassemblers(FNodeInfo& NodeInfo)
 {
 	SCOPED_MESSAGING_TRACE(FUdpMessageProcessor_UpdateReassemblers);
+
 	FUdpMessageSegment::FHeader Header
 	{
 		FMath::Max(NodeInfo.ProtocolVersion, (uint8)11),	// Header.ProtocolVersion, AcknowledgeSegments are version 11 and onward segment
@@ -1618,6 +1620,7 @@ bool FUdpMessageProcessor::UpdateReassemblers(FNodeInfo& NodeInfo)
 		return true;
 	}
 
+	const FUdpMessagingTime StaleLastSegmentTime = FUdpMessagingTime::Now() - StaleReassemblyInterval;
 	const uint32 MaxSendRate = GetMaxSendRate();
 	double DeltaSeconds = DeltaTime.GetTotalSeconds();
 	const int32 MaxSendRateDelta = MaxSendRate * DeltaSeconds;
@@ -1635,11 +1638,22 @@ bool FUdpMessageProcessor::UpdateReassemblers(FNodeInfo& NodeInfo)
 			It.RemoveCurrent();
 			continue;
 		}
+		else if (ReassembledMessage->IsPendingDelivery())
+		{
+			// Leave it in the list, but skip further processing.
+			continue;
+		}
+
+		SCOPED_MESSAGING_TRACE(FUdpMessageProcessor_UpdateReassemblers_Message);
+
+		const FUdpMessagingTime MaxAckDelay = FUdpMessagingTime::Now() - FTimespan::FromMilliseconds(5);
+		const bool bShouldAck = (ReassembledMessage->GetLastAckTime() < MaxAckDelay) || (ReassembledMessage->NumPendingAcknowledgements() >= MaxAckNum);
 
 		int BytesSent = 0;
-		while (ReassembledMessage->HasPendingAcknowledgements() && BytesSent < MaxNodeByteSend)
+		while (bShouldAck && ReassembledMessage->NumPendingAcknowledgements() > 0)
 		{
-			TArray<uint32> PendingAcknowledgments = ReassembledMessage->TakePendingAcknowledgments(MaxAckNum);
+			SCOPED_MESSAGING_TRACE(FUdpMessageProcessor_UpdateReassemblers_Acks);
+			TArray<uint32> PendingAcknowledgments = ReassembledMessage->TakePendingAcknowledgments(MaxAckNum, FUdpMessagingTime::Now());
 			const int32 AckCount = PendingAcknowledgments.Num();
 
 			FUdpMessageSegment::FAcknowledgeSegmentsChunk AcknowledgeChunk(It.Key()/*MessageId*/, MoveTemp(PendingAcknowledgments)/*Segments*/);
@@ -1652,7 +1666,9 @@ bool FUdpMessageProcessor::UpdateReassemblers(FNodeInfo& NodeInfo)
 
 			UE_LOG(LogUdpMessaging, Verbose, TEXT("FUdpMessageProcessor::UpdateReassemblers: Sending EUdpMessageSegments::AcknowledgeSegments for %d segments for message %d from %s"), AckCount, It.Key(), *ReassembledMessage->Describe());
 
-			if (!SocketSender->Send(Writer, NodeInfo.Endpoint))
+			const bool bAllowSegmentation_false = false;
+			const bool bHighPriority_true = true;
+			if (!SocketSender->Send(Writer, NodeInfo.Endpoint, bAllowSegmentation_false, bHighPriority_true))
 			{
 				UE_LOG(LogUdpMessaging, Error, TEXT("FUdpMessageProcessor::UpdateReassemblers: error sending EUdpMessageSegments::AcknowledgeSegments from %s"), *NodeInfo.NodeId.ToString());
 
@@ -1671,7 +1687,7 @@ bool FUdpMessageProcessor::UpdateReassemblers(FNodeInfo& NodeInfo)
 		}
 
 		// Remove stale reassembled message if they aren't reliable or are marked delivered
-		if (ReassembledMessage->GetLastSegmentTime() + StaleReassemblyInterval <= CurrentTime &&
+		if (ReassembledMessage->GetLastSegmentTime() <= StaleLastSegmentTime &&
 			(!EnumHasAnyFlags(ReassembledMessage->GetFlags(), EMessageFlags::Reliable)))
 		{
 			if (!ReassembledMessage->IsDelivered() && !ReassembledMessage->IsPendingDelivery())
diff --git a/Engine/Plugins/Messaging/UdpMessaging/Source/UdpMessaging/Private/Transport/UdpMessageProcessor.h b/Engine/Plugins/Messaging/UdpMessaging/Source/UdpMessaging/Private/Transport/UdpMessageProcessor.h
--- a/Engine/Plugins/Messaging/UdpMessaging/Source/UdpMessaging/Private/Transport/UdpMessageProcessor.h
+++ b/Engine/Plugins/Messaging/UdpMessaging/Source/UdpMessaging/Private/Transport/UdpMessageProcessor.h
@@ -19,8 +19,10 @@
 #include "UObject/WeakObjectPtrTemplates.h"
 
 #include "Shared/UdpMessagingSettings.h"
+#include "Shared/UdpMessagingTime.h"
 #include "UdpMessageSegmenter.h"
 #include "UdpMessagingPrivate.h"
+#include "UdpMessagingTracing.h"
 
 
 class IMessageContext;
@@ -61,7 +63,15 @@ struct FSentData
 	uint64 SequenceNumber = 0;
 	bool bIsReliable = false;
 
-	FDateTime TimeSent;
+	FUdpMessagingTime TimeSent;
+};
+
+/** Lightweight tracking of segment transmit times to expedite maintenance of send window. */
+struct FSegmentExpiry
+{
+	FUdpMessagingTime TimeSent;
+	int32 MessageId = 0;
+	uint32 SegmentNumber = 0;
 };
 
 /** Returned by the segment processor. It provides details on what was sent and how it was sent. */
@@ -130,8 +140,8 @@ class FUdpMessageProcessor
 		/** Holds the node's IP endpoint. */
 		FIPv4Endpoint Endpoint;
 
-		/** Holds the time at which the last Hello was received. */
-		FDateTime LastSegmentReceivedTime;
+		/** Holds the time at which we last received any segment from this node. */
+		FUdpMessagingTime LastSegmentReceivedTime;
 
 		/** Holds the endpoint's node identifier. */
 		FGuid NodeId;
@@ -156,30 +166,41 @@ class FUdpMessageProcessor
 		/** Holds of queue of MessageIds to send for unreliable messages. They are processed in round-robin fashion. */
 		FWorkQueue UnreliableWorkQueue;
 
-		/* Holds the list of messages that are still pending acknowledgement or overflow if overcommitted on work queue. */
+		/* Holds the list of messages that are still pending acknowledgement. */
 		TArray<int32> OverflowForPendingAck;
 
-		/** Last time we issued a warning about work queues. */
-		FDateTime LastWorkQueueFullMessage;
-
-		/** The number of seconds we have been in a send full state. */
-		double SecondsInFullQueueState = 0;
-
 		/** A map from sequence id to information about what was sent. The size of this map is the size of our sending window. */
 		TMap<uint64, FSentData> InflightSegments;
 
+		/** Send timings for outgoing reliable segments. */
+		TRingBuffer<FSegmentExpiry> ReliableExpiry;
+
+		/** Send timings for outgoing unreliable segments. */
+		TRingBuffer<FSegmentExpiry> UnreliableExpiry;
+
 		/** Unique sequence id for every packet sent. */
 		uint64 SequenceId = 0;
 
-		/** Average time from when the packet was sent to ack received. */
-		FTimespan AvgRoundTripTime = FTimespan::FromMilliseconds(20);
-
 		/** Maximum number of segments we can have in flight. Start with a middle of the road value and adjust based on network performance. */
 		uint64 WindowSize = 1024;
 
 		/** Maximum window size that can be set. */
 		uint64 MaxWindowSize = 2048;
 
+		/** A smoothed calculated RTT that is effectively the weighted average. */
+		double SmoothedRTT = -1.0;
+
+		/** Measured variance of the RTT to be used to calculate a retransmission timeout. */
+		double RTTVariance = 0.0;
+
+		/**
+		 * Set a high retransmission timeout on init until we start measuring speed and overall jitter.
+		 * The default is 1 second. In other words, if we don't get an ack within 1 second, we assume it was
+		 * lost and thus should be marked for retransmission. This value will adjust lower as we compute
+		 * new data especially on a fast network.
+		 */
+		double RetransmissionTimeout = 1.0f;
+
 		/** Various transport statistics for this endpoint */
 		FMessageTransportStatistics Statistics;
 
@@ -205,15 +226,14 @@ class FUdpMessageProcessor
 		}
 
 		/**
-		 * Update the segmenter with the given acts. If all acks have been received then remove it from the list of segmenters.
+		 * Update the segmenter with the given acks. If all acks have been received then remove it from the list of segmenters.
 		 */
-		TSharedPtr<FUdpMessageSegmenter> MarkAcksOnSegmenter(int32 MessageId, const TArray<uint32>& Segments, const FDateTime& InCurrentTime)
+		TSharedPtr<FUdpMessageSegmenter> MarkAcksOnSegmenter(int32 MessageId, const TArray<uint32>& Segments, const FUdpMessagingTime InCurrentTime)
 		{
 			TSharedPtr<FUdpMessageSegmenter> Segmenter = Segmenters.FindRef(MessageId);
 			if (Segmenter.IsValid())
 			{
 				Segmenter->MarkAsAcknowledged(Segments);
-				Segmenter->UpdateSentTime(InCurrentTime);
 				if (Segmenter->IsSendingComplete() && Segmenter->AreAcknowledgementsComplete())
 				{
 					UE_LOG(LogUdpMessaging, Verbose, TEXT("Segmenter for %s is now complete. Removing"), *NodeId.ToString());
@@ -243,31 +263,47 @@ class FUdpMessageProcessor
 		}
 
 		/**
-		 * Update the calcuated round trip time based on the computed span of inflight packet.
+		 * Update the calculated round trip time based on the computed span of inflight packet.
 		 */
 		void CalculateNewRTT(const FTimespan& NewSpan)
 		{
-			const float Weight = 0.6f;
-			Statistics.AverageRTT = Weight * Statistics.AverageRTT + (1-Weight) * NewSpan;
-		}
+			constexpr double ClockGranularity = 0.005; // "G" = 5ms
 
-		/**
-		 * Remove any unreliables that are too old.  Too old is based on the average round trip time of a packet.
-		 */
-		void RemoveOldUnreliables(const FDateTime& InCurrentTime)
-		{
-			// Remove any unreliable packets that fall outside of our average RTT
-			RemoveInflightIf([InCurrentTime, this](const FSentData& Data)
+			// The computed round trip time for this time.
+			const double RTT = NewSpan.GetTotalSeconds();
+
+			if (SmoothedRTT > 0)
 			{
-				return ((InCurrentTime - Data.TimeSent) > Statistics.AverageRTT) && !Data.bIsReliable;
-			});
+				// Calculate the smoothed variance of the RTT so that we can use it in computing a timeout value.
+				const double VarianceWeight = 1 / 4.0; // "beta"
+				RTTVariance = (1 - VarianceWeight) * RTTVariance + VarianceWeight * FMath::Abs(RTT - SmoothedRTT);
+
+				// Use a smoothed RTT. 1/8 is the value TCP typically uses.
+				const double Weight = 1 / 8.0; // "alpha"
+				SmoothedRTT = (1 - Weight) * SmoothedRTT + Weight * RTT;
+			}
+			else
+			{
+				// Calculate initial values based on the first RTT. 
+				SmoothedRTT = RTT;
+				RTTVariance = SmoothedRTT / 2;
+			}
+
+			Statistics.AverageRTT = FTimespan::FromSeconds(SmoothedRTT);
+
+			// Our transmission timeout becomes the smoothed RTT plus the calculated jitter. 
+			RetransmissionTimeout = SmoothedRTT + FMath::Max(ClockGranularity, 4 * RTTVariance);
+
+			static constexpr double MinRTO = 0.020;
+			static constexpr double MaxRTO = 3.000;
+			RetransmissionTimeout = FMath::Clamp(RetransmissionTimeout, MinRTO, MaxRTO);
 		}
 
 		/** 
 		 * Compute the new window size for the connection.  We use a AIMD algorithm:
 		 * https://en.wikipedia.org/wiki/Additive_increase/multiplicative_decrease
 		 */
-		void ComputeWindowSize(uint32 NumAcks, uint32 SegmentLoss)
+		void ComputeWindowSize(uint32 NumAcks, uint32 SegmentLoss = 0)
 		{
 			const uint64 MinimumWindowSize = 64;
 			const uint64 MaximumWindowSize = MaxWindowSize;
@@ -285,6 +321,31 @@ class FUdpMessageProcessor
 			Statistics.WindowSize = WindowSize;
 		}
 
+		/** Called on transmit. Keeps expiry queues in sync with InflightSegments map. */
+		void AddInflight(const FSentData& SentData)
+		{
+			const int32 MessageId = SentData.MessageId;
+			const uint32 SegmentNumber = SentData.SegmentNumber;
+
+			const uint64 Key = Remap(MessageId, SegmentNumber);
+			InflightSegments.Add(Key, SentData);
+
+			TRingBuffer<FSegmentExpiry>& ExpiryQueue = SentData.bIsReliable ? ReliableExpiry : UnreliableExpiry;
+			ExpiryQueue.Emplace(FSegmentExpiry{ .TimeSent = SentData.TimeSent,
+				.MessageId = MessageId, .SegmentNumber = SegmentNumber });
+		}
+
+		/** Undo a previous AddInflight. This frees up send window, but leaves the harmless stale expiry. */
+		void RemoveInflight(const FSentData& SentData)
+		{
+			const int32 MessageId = SentData.MessageId;
+			const uint32 SegmentNumber = SentData.SegmentNumber;
+
+			const uint64 Key = Remap(MessageId, SegmentNumber);
+			InflightSegments.Remove(Key);
+			// We leave stale expiry entries because it's cheap to lazily ignore them later.
+		}
+
 		/**
 		 * Remove all segments in the inflight buffer that match the given MessageId.
 		 */
@@ -302,41 +363,45 @@ class FUdpMessageProcessor
 		 * Make a given MessageId as complete and recompute the desired window size based
 		 * on any loss / acks received.
 		 */
-		void MarkComplete(int32 MessageId, const FDateTime& InCurrentTime)
+		void MarkComplete(int32 MessageId, const FUdpMessagingTime InCurrentTime)
 		{
-			uint32 AckSegments = RemoveAllInflightFromMessageId(MessageId);
-			uint32 SegmentLoss = RemoveLostSegments(InCurrentTime);
-			ComputeWindowSize(AckSegments, SegmentLoss);
+			const uint32 AckSegments = RemoveAllInflightFromMessageId(MessageId);
+			const uint32 LostSegmentCount = ExpireSegments(InCurrentTime);
+			ComputeWindowSize(AckSegments, LostSegmentCount);
 		}
 
 		/**
 		 * For a given list of segments to acknowledge record how long it took to receive a response and update our
 		 * average round trip time. We also use this opportunity to calculate any loss and update our window size.
 		 */
-		void MarkAcks(int32 MessageId, const TArray<uint32>& Segments, const FDateTime& InCurrentTime)
+		void MarkAcks(int32 MessageId, const TArray<uint32>& Segments, const FUdpMessagingTime InCurrentTime)
 		{
 			TSharedPtr<FUdpMessageSegmenter> Segmenter = MarkAcksOnSegmenter(MessageId, Segments, InCurrentTime);
-			uint64 MaxSequenceId = 0;
 			FTimespan MaxSpan(0);
 			uint32 Acks = 0;
 			for (uint32 SegmentId : Segments)
 			{
-				uint64 Id = Remap(MessageId,SegmentId);
+				const uint64 Id = Remap(MessageId, SegmentId);
 				FSentData SentData;
-				if (InflightSegments.RemoveAndCopyValue(Id,SentData))
+				if (InflightSegments.RemoveAndCopyValue(Id, SentData))
 				{
-					MaxSpan = FGenericPlatformMath::Max<FTimespan>(InCurrentTime - SentData.TimeSent, MaxSpan);
-					MaxSequenceId = FGenericPlatformMath::Max<uint64>(SentData.SequenceNumber,MaxSequenceId);
+					// Karn's algorithm: retransmitted segments are ineligible for RTT measurement.
+					if (Segmenter && !Segmenter->WasRetransmitted(SegmentId))
+					{
+						MaxSpan = FGenericPlatformMath::Max<FTimespan>(InCurrentTime - SentData.TimeSent, MaxSpan);
+					}
+
 					++Acks;
 				}
 			}
+
 			if (MaxSpan > 0)
 			{
 				CalculateNewRTT(MaxSpan);
 			}
 
-			uint32 SegmentLoss = RemoveLostSegments(InCurrentTime);
-			ComputeWindowSize(Acks, SegmentLoss);
+			const uint32 LostSegmentCount = ExpireSegments(InCurrentTime);
+			ComputeWindowSize(Acks, LostSegmentCount);
 		}
 
 		/**
@@ -356,36 +421,64 @@ class FUdpMessageProcessor
 		 * using 2 times the average round trip time.  If an ack has not been received in that time frame then it is
 		 * lost and we must resend it.
 		 */
-		uint32 RemoveLostSegments(const FDateTime& InCurrentTime)
+		uint32 ExpireSegments(const FUdpMessagingTime InCurrentTime)
 		{
-			uint32 SegmentsLost = 0;
+			SCOPED_MESSAGING_TRACE(FNodeInfo_ExpireSegments);
 
-			// We use 2 times the average round trip time to determine if a segment has been lost.
-			RemoveInflightIf([&SegmentsLost, InCurrentTime, this](const FSentData& Data)
+			auto ExpireUpTo = [this]
+			(TRingBuffer<FSegmentExpiry>& ExpiryQueue, FUdpMessagingTime ExpiryTime, bool bMarkLoss) -> uint32
 			{
-				const bool bIsLost = ((InCurrentTime - Data.TimeSent) > 2*AvgRoundTripTime) && Data.bIsReliable;
-				if (bIsLost)
+				uint32 SegmentsExpired = 0;
+				while (ExpiryQueue.Num() && ExpiryQueue.First().TimeSent <= ExpiryTime)
 				{
-					MarkSegmenterSegmentLoss(Data);
-					SegmentsLost++;
+					const FSegmentExpiry Expiry = ExpiryQueue.PopFrontValue();
+
+					const uint64 ExpiredKey = Remap(Expiry.MessageId, Expiry.SegmentNumber);
+					const FSentData* ExpiredData = InflightSegments.Find(ExpiredKey);
+
+					// This may have already been removed for other reasons; that's OK.
+					// If the timestamp doesn't match, we retransmitted; ignore this stale expiry.
+					if (!ExpiredData || ExpiredData->TimeSent != Expiry.TimeSent)
+					{
+						continue;
+					}
+
+					if (bMarkLoss)
+					{
+						MarkSegmenterSegmentLoss(*ExpiredData);
+					}
+
+					++SegmentsExpired;
+					InflightSegments.Remove(ExpiredKey);
 				}
-				return bIsLost;
-			});
 
-			Statistics.TotalBytesLost += SegmentsLost * MaxSegmentSize;
-			Statistics.PacketsLost += SegmentsLost;
-			return SegmentsLost;
+				return SegmentsExpired;
+			};
+
+			// Expire reliable segments.
+			const FTimespan ReliableTimeout = FTimespan::FromSeconds(RetransmissionTimeout);
+			const FUdpMessagingTime ReliableSentExpiryTime = InCurrentTime - ReliableTimeout;
+			const uint32 ReliableSegmentsLost = ExpireUpTo(ReliableExpiry, ReliableSentExpiryTime, true);
+
+			// Expire unreliable segments. They are not ack-eliciting, so 1/2 timeout is intended to account for send congestion.
+			const FTimespan UnreliableTimeout = ReliableTimeout / 2;
+			const FUdpMessagingTime UnreliableSentExpiryTime = InCurrentTime - UnreliableTimeout;
+			ExpireUpTo(UnreliableExpiry, UnreliableSentExpiryTime, false);
+
+			Statistics.TotalBytesLost += ReliableSegmentsLost * MaxSegmentSize;
+			Statistics.PacketsLost += ReliableSegmentsLost;
+			return ReliableSegmentsLost;
 		}
 
 		/**
 		 * Does this node have any segmenters that can send data for the current time frame.
 		 */
-		bool HasSegmenterThatCanSend(const FDateTime& InCurrentTime)
+		bool HasSegmenterThatCanSend(const FUdpMessagingTime InCurrentTime)
 		{
 			using FSegmenterTuple = TTuple<int32,TSharedPtr<FUdpMessageSegmenter>>;
 			for (FSegmenterTuple& Segmenter : Segmenters)
 			{
-				if (Segmenter.Value->IsInitialized() && Segmenter.Value->NeedSending(InCurrentTime))
+				if (Segmenter.Value->IsInitialized() && !Segmenter.Value->IsSendingComplete())
 				{
 					return true;
 				}
@@ -848,8 +941,8 @@ protected:
 		/** Holds the node's IP endpoint. */
 		FIPv4Endpoint Endpoint;
 
-		/** Holds the time at which the last Hello was received. */
-		FDateTime LastSegmentReceivedTime;
+		/** Holds the time at which we last received any segment from this node. */
+		FUdpMessagingTime LastSegmentReceivedTime;
 
 		/** Holds the endpoint's node identifier. */
 		FGuid NodeId;
@@ -913,7 +1006,7 @@ private:
 	/** Handles a communication error to a particular endpoint.*/
 	void HandleSocketError(const TCHAR* ErrorReference, const FNodeInfo& NodeInfo) const;
 
-	/** Checks all known nodes to see if any segmenters have NeedSending set to true. */
+	/** Checks all known nodes to see if any have both segmenters and send windows ready to transmit. */
 	bool MoreToSend();
 
 	/** Consume an Outbound Message for processing. */
@@ -928,9 +1021,6 @@ private:
 	/** Holds the hello sender. */
 	TUniquePtr<FUdpMessageBeacon> Beacon;
 
-	/** Holds the current time. */
-	FDateTime CurrentTime;
-
 	/** Holds the delta time between two ticks. */
 	FTimespan DeltaTime;
 
diff --git a/Engine/Plugins/Messaging/UdpMessaging/Source/UdpMessaging/Private/Transport/UdpMessageProcessorLogger.cpp b/Engine/Plugins/Messaging/UdpMessaging/Source/UdpMessaging/Private/Transport/UdpMessageProcessorLogger.cpp
--- a/Engine/Plugins/Messaging/UdpMessaging/Source/UdpMessaging/Private/Transport/UdpMessageProcessorLogger.cpp
+++ b/Engine/Plugins/Messaging/UdpMessaging/Source/UdpMessaging/Private/Transport/UdpMessageProcessorLogger.cpp
@@ -9,7 +9,6 @@
 #include "Containers/SpscQueue.h"
 #include "Interfaces/IPv4/IPv4Endpoint.h"
 
-#include "Misc/DateTime.h"
 #include "Misc/Guid.h"
 #include "Misc/ScopeExit.h"
 #include "Misc/Paths.h"
@@ -38,8 +37,8 @@ struct FNodeInfoStats
 	/** Holds the node's IP endpoint. */
 	FIPv4Endpoint Endpoint;
 
-	/** Holds the time at which the last Hello was received. */
-	FDateTime LastSegmentReceivedTime;
+	/** Holds the time at which we last received any segment from this node. */
+	FUdpMessagingTime LastSegmentReceivedTime;
 
 	/** Holds the endpoint's node identifier. */
 	FGuid NodeId;
@@ -160,7 +159,14 @@ struct FUdpMessageProcessorLoggerTyped
 		FString InstanceName = FApp::GetInstanceName();
 		FString EndpointStr = Endpoint.ToString().Replace(TEXT("."), TEXT("_")).Replace(TEXT(":"), TEXT("p"));
 		const FString Filename = FString::Printf(TEXT("%s-%s-%s-%s.csv"),*InstanceName, *EndpointStr, *UE::UdpMessaging::Private::GetOutputName(this), *FDateTime::Now().ToString());
-		const FString FullPath = FPaths::ProjectLogDir() / TEXT("UdpMessaging") / Filename;
+		const FString DestDir = FPaths::ProjectLogDir() / TEXT("UdpMessaging");
+		const FString FullPath = DestDir / Filename;
+
+		if (!IFileManager::Get().MakeDirectory(*DestDir, true))
+		{
+			UE_LOG(LogUdpMessaging, Warning, TEXT("Error creating log directory: %s"), *DestDir);
+		}
+
 		CSVArchive = TUniquePtr<FArchive>(IFileManager::Get().CreateFileWriter(*FullPath, EFileWrite::FILEWRITE_AllowRead));
 		if (CSVArchive)
 		{
@@ -425,7 +431,7 @@ namespace UE::UdpMessaging::Private
 		TStringBuilder<512> Builder;
 		Builder.Append(Stat.Endpoint.ToString());
 		Builder.Append(TEXT(","));
-		Builder.Append(Stat.LastSegmentReceivedTime.ToString());
+		Builder.Appendf(TEXT("%.3f"), Stat.LastSegmentReceivedTime.RawValue);
 		Builder.Append(TEXT(","));
 		Builder.Append(Stat.NodeId.ToString());
 		Builder.Append(TEXT(","));
diff --git a/Engine/Plugins/Messaging/UdpMessaging/Source/UdpMessaging/Private/Transport/UdpMessageSegmenter.cpp b/Engine/Plugins/Messaging/UdpMessaging/Source/UdpMessaging/Private/Transport/UdpMessageSegmenter.cpp
--- a/Engine/Plugins/Messaging/UdpMessaging/Source/UdpMessaging/Private/Transport/UdpMessageSegmenter.cpp
+++ b/Engine/Plugins/Messaging/UdpMessaging/Source/UdpMessaging/Private/Transport/UdpMessageSegmenter.cpp
@@ -220,6 +220,10 @@ void FUdpMessageSegmenter::MarkForRetransmission(uint32 SegmentId)
 
 		++PendingSendSegmentsCount;
 		PendingSendSegments[SegmentId] = true;
+
+		RetransmittedSegments.Add(SegmentId);
+		++RetransmitCount;
+
 		// Note - we don't need to clear acknowledgments. If any segment is in transit and acknowledged after this
 		// call we'll do that and if possible stop the pending send, and if not we don't need to wait for the ack
 	}
@@ -240,6 +244,17 @@ void FUdpMessageSegmenter::MarkForRetransmission()
 {
 	UE_LOG(LogUdpMessaging, Verbose, TEXT("Marking all %d segments for retransmission"), PendingSendSegments.Num());
 
+	for (TBitArray<>::FConstIterator BIt(AcknowledgeSegments, EarliestUnackedSegmentId); BIt; ++BIt)
+	{
+		const bool bAcknowledged = BIt.GetValue();
+		const int32 SegmentId = BIt.GetIndex();
+		if (!bAcknowledged && !PendingSendSegments[SegmentId])
+		{
+			RetransmittedSegments.Add(SegmentId);
+			++RetransmitCount;
+		}
+	}
+
 	// mark all segments to be resent and clear any pending state
 	PendingSendSegments.Init(true, PendingSendSegments.Num());
 	PendingSendSegmentsCount = PendingSendSegments.Num();
@@ -248,67 +263,12 @@ void FUdpMessageSegmenter::MarkForRetransmission()
 	// call we'll do that and if possible stop the pending send, and if not we don't need to wait for the ack
 }
 
-TAutoConsoleVariable<int32> CVarSegmenterTimeout(
-	TEXT("MessageBus.UDP.SegmenterTimeout"),
-	100,
-	TEXT("The number of milliseconds to wait for a segment to be acknowledge before attempting to resend. Values are clamped between 10 and 1000"),
-	ECVF_Default
-);
-TAutoConsoleVariable<int32> CVarSegmenterMaxResends(
-	TEXT("MessageBus.UDP.SegmenterMaxResends"),
-	16,
-	TEXT("The number of attempts to resend data to an endpoint. Values are clamped between 1 and 100."),
-	ECVF_Default
-);
-
-bool FUdpMessageSegmenter::NeedSending(const FDateTime& CurrentTime)
+bool FUdpMessageSegmenter::WasRetransmitted(uint32 SegmentId) const
 {
-	// still have outstanding segments
-	if (PendingSendSegmentsCount > 0)
-	{
-		return true;
-	}
-
-	const uint16 MaxNumResends = FMath::Clamp(CVarSegmenterMaxResends.GetValueOnAnyThread(), 1, 100);
-	if (AreAcknowledgementsComplete() == false && SentNumber > MaxNumResends)
+	if (RetransmitCount == 0)
 	{
-		UE_LOG(LogUdpMessaging, Warning, TEXT("Gave up sending with %d outstanding acks. Attempted %d tries."),
-			   AcknowledgeSegments.Num() - AcknowledgeSegmentsCount,
-			   MaxNumResends
-		);
-		SerializedMessage->UpdateState(EUdpSerializedMessageState::Invalid);
 		return false;
 	}
 
-	const int32 NumMs = FMath::Clamp(CVarSegmenterTimeout.GetValueOnAnyThread(), 10, 5000);
-	const FTimespan SendInterval = FTimespan::FromMilliseconds(NumMs);
-	if (AreAcknowledgementsComplete() == false
-		&& LastSentTime + SendInterval <= CurrentTime)
-	{
-		// We have gone through a period of time where packets or acks may have been lost,
-		// so resend any segments that have yet to be acknowledged
-		for (TBitArray<>::FConstIterator BIt(AcknowledgeSegments, EarliestUnackedSegmentId); BIt; ++BIt)
-		{
-			const int32 Index = BIt.GetIndex();
-			if (BIt.GetValue() == false && !PendingSendSegments[Index])
-			{
-				PendingSendSegments[Index] = true;
-				++PendingSendSegmentsCount;
-			}
-		}
-		LastSentTime = CurrentTime;
-		++SentNumber;
-		UE_LOG(LogUdpMessaging, Warning, TEXT("Waiting for ack too long. Re-sending %d segments. Attempt %d out of %d tries."),
-			   AcknowledgeSegments.Num() - AcknowledgeSegmentsCount,
-			   SentNumber,
-			   MaxNumResends);
-		return true;
-	}
-
-	return false;
-}
-
-void FUdpMessageSegmenter::UpdateSentTime(const FDateTime& CurrentTime)
-{
-	LastSentTime = CurrentTime;
+	return RetransmittedSegments.Contains(SegmentId);
 }
diff --git a/Engine/Plugins/Messaging/UdpMessaging/Source/UdpMessaging/Private/Transport/UdpMessageSegmenter.h b/Engine/Plugins/Messaging/UdpMessaging/Source/UdpMessaging/Private/Transport/UdpMessageSegmenter.h
--- a/Engine/Plugins/Messaging/UdpMessaging/Source/UdpMessaging/Private/Transport/UdpMessageSegmenter.h
+++ b/Engine/Plugins/Messaging/UdpMessaging/Source/UdpMessaging/Private/Transport/UdpMessageSegmenter.h
@@ -2,11 +2,11 @@
 
 #pragma once
 
+#include "Containers/Set.h"
 #include "Shared/UdpMessageSegment.h"
+#include "Shared/UdpMessagingTime.h"
 #include "Templates/SharedPointer.h"
-#include "Misc/DateTime.h"
 
-struct FTimespan;
 
 // IMessageContext forward declaration
 enum class EMessageFlags : uint32;
@@ -240,19 +240,8 @@ public:
 	 */
 	void MarkForRetransmission(const TArray<uint16>& Segments);
 
-	/**
-	 * Checks if this segmenter needs to send segments
-	 *
-	 * @return true if the segmenter needs to send
-	 */
-	bool NeedSending(const FDateTime& CurrentTime);
-
-	/**
-	 * Update the last sent time and increment the sent number for this segmenter
-	 *
-	 * @param CurrentTime the time to update to.
-	 */
-	void UpdateSentTime(const FDateTime& CurrentTime);
+	/** Return whether a segment has been retransmitted. */
+	bool WasRetransmitted(uint32 SegmentId) const;
 
 	/**
 	 * Checks whether the serialization of the message to segment and send completed. The serialization may be performing asynchronously and can may succeed or fail. The message cannot be segmented
@@ -286,12 +275,12 @@ private:
 	/** Holds the segment size (including the FHeader and FDataChunk fields). */
 	uint16 SegmentSize = 0;
 
-	/** Holds the number of time we sent the segments */
-	uint16 SentNumber = 0;
-
-	/** Holds the time at which we last sent */
-	FDateTime LastSentTime;
+	/** Holds the number of retransmits. */
+	uint16 RetransmitCount = 0;
 
 	/** Holds the message. */
 	TSharedPtr<FUdpSerializedMessage, ESPMode::ThreadSafe> SerializedMessage;
+
+	/** Keep a record of any segments we've retransmitted. */
+	TSet<uint32> RetransmittedSegments;
 };
diff --git a/Engine/Plugins/Messaging/UdpMessaging/Source/UdpMessaging/Private/Transport/UdpMessageTransport.cpp b/Engine/Plugins/Messaging/UdpMessaging/Source/UdpMessaging/Private/Transport/UdpMessageTransport.cpp
--- a/Engine/Plugins/Messaging/UdpMessaging/Source/UdpMessaging/Private/Transport/UdpMessageTransport.cpp
+++ b/Engine/Plugins/Messaging/UdpMessaging/Source/UdpMessaging/Private/Transport/UdpMessageTransport.cpp
@@ -476,7 +476,7 @@ void FUdpMessageTransport::HandleProcessorMessageReassembled(FUdpReassembledMess
 		else
 		{
 			TWeakObjectPtr<UScriptStruct> TypeInfo = ReassembledMessagePtr->GetMessageTypeInfo();
-			FString TypeInfoStr = TypeInfo.IsValid() ? TypeInfo->GetName() : TEXT("UKNOWN TYPE");
+			FString TypeInfoStr = TypeInfo.IsValid() ? TypeInfo->GetName() : TEXT("UNKNOWN TYPE");
 			UE_LOG(LogUdpMessaging, Verbose, TEXT("Failed to deserialize message of type %s from %s"), *TypeInfoStr, *NodeId.ToString());
 		}
 		// We always mark the message as delivered even if we can't deserialize it. There may be some messages that can't be deserialized because of missing type information.
diff --git a/Engine/Plugins/Messaging/UdpMessaging/Source/UdpMessaging/Private/Transport/UdpReassembledMessage.h b/Engine/Plugins/Messaging/UdpMessaging/Source/UdpMessaging/Private/Transport/UdpReassembledMessage.h
--- a/Engine/Plugins/Messaging/UdpMessaging/Source/UdpMessaging/Private/Transport/UdpReassembledMessage.h
+++ b/Engine/Plugins/Messaging/UdpMessaging/Source/UdpMessaging/Private/Transport/UdpReassembledMessage.h
@@ -4,7 +4,7 @@
 
 #include "IMessageContext.h"
 #include "Interfaces/IPv4/IPv4Endpoint.h"
-#include "Misc/DateTime.h"
+#include "Shared/UdpMessagingTime.h"
 #include "UdpMessagingPrivate.h"
 
 // IMessageContext forward declaration
@@ -107,7 +107,7 @@ public:
 	 *
 	 * @return Last receive time.
 	 */
-	FDateTime GetLastSegmentTime() const
+	FUdpMessagingTime GetLastSegmentTime() const
 	{
 		return LastSegmentTime;
 	}
@@ -194,24 +194,32 @@ public:
 	}
 
 	/**
-	 * Checks whether this message has pending acknowledgments to send.
-	 *
-	 * @return true if the message has pending acknowledgments, false otherwise.
+	 * @return How many acknowledgments are currently pending send.
+	 */
+	int32 NumPendingAcknowledgements() const
+	{
+		return PendingAcknowledgments.Num();
+	}
+
+	/**
+	 * @return The time at which any acknowledgments were last transmitted for this message.
 	 */
-	bool HasPendingAcknowledgements() const
+	FUdpMessagingTime GetLastAckTime() const
 	{
-		return PendingAcknowledgments.Num() > 0;
+		return LastAckTime;
 	}
 
 	/**
-	 * Gets at most the specified number of pending acknowledgments and clears them.
+	 * Gets at most the specified number of pending acknowledgments and clears them, updating LastAckTime.
 	 *
 	 * @return The array of pending acknowledgments removed.
 	 */
-	TArray<uint32> TakePendingAcknowledgments(int32 MaxAckNum)
+	TArray<uint32> TakePendingAcknowledgments(int32 MaxAckNum, FUdpMessagingTime CurrentTime)
 	{
 		check(MaxAckNum > 0);
 
+		LastAckTime = CurrentTime;
+
 		TArray<uint32> Temp;
 		if (PendingAcknowledgments.Num() < MaxAckNum)
 		{
@@ -278,7 +286,7 @@ public:
 	 * @param SegmentData The segment data.
 	 * @param CurrentTime The current time.
 	 */
-	void Reassemble(int32 SegmentNumber, int32 SegmentOffset, const TArray<uint8>& SegmentData, const FDateTime& CurrentTime)
+	void Reassemble(int32 SegmentNumber, int32 SegmentOffset, const TArray<uint8>& SegmentData, const FUdpMessagingTime CurrentTime)
 	{
 		if (IsPendingDelivery() || IsMalformed() || !IsInitialized() || SegmentData.Num() == 0)
 		{
@@ -310,7 +318,7 @@ public:
 
 		LastSegmentTime = CurrentTime;
 
-		// Send acknowledgement for unreliable messages
+		// Send acknowledgment for reliable messages
 		if (EnumHasAnyFlags(GetFlags(), EMessageFlags::Reliable))
 		{
 			PendingAcknowledgments.AddUnique(SegmentNumber);
@@ -365,7 +373,7 @@ private:
 	TArray<uint8> Data;
 
 	/** Holds the time at which the last segment was received. */
-	FDateTime LastSegmentTime;
+	FUdpMessagingTime LastSegmentTime;
 
 	/** Holds an array of bits that indicate which segments still need to be received. */
 	TBitArray<> PendingSegments;
@@ -376,6 +384,9 @@ private:
 	/** Acknowledgment yet to be sent about segments we received */
 	TArray<uint32> PendingAcknowledgments;
 
+	/** Holds the time at which any pending acknowledgments were last transmitted. */
+	FUdpMessagingTime LastAckTime;
+
 	/** Holds the number of bytes received so far. */
 	int32 ReceivedBytes;
 
