UdpMessaging: Prevent double delivery of retransmitted messages.

This change also removes the `FNodeInfo::ResetIfRestarted` method, which originated from a time when segment headers did not contain the source node ID, and the processor instead maintained a mapping from FIPv4Endpoint to FGuid. In the current version of the code, this no longer makes sense, and an FNodeInfo should never be reused. Change the Hello/Ping/Pong handlers to assert that the payload is consistent with the header.

diff --git a/Engine/Plugins/Messaging/UdpMessaging/Source/UdpMessaging/Private/Transport/UdpMessageProcessor.cpp b/Engine/Plugins/Messaging/UdpMessaging/Source/UdpMessaging/Private/Transport/UdpMessageProcessor.cpp
--- a/Engine/Plugins/Messaging/UdpMessaging/Source/UdpMessaging/Private/Transport/UdpMessageProcessor.cpp
+++ b/Engine/Plugins/Messaging/UdpMessaging/Source/UdpMessaging/Private/Transport/UdpMessageProcessor.cpp
@@ -783,6 +783,13 @@ void FUdpMessageProcessor::ProcessDataSegment(FInboundSegment& Segment, FNodeInf
 		return;
 	}
 
+	if (NodeInfo.RecentlyDeliveredMessages.Contains(DataChunk.MessageId))
+	{
+		UE_LOG(LogUdpMessaging, VeryVerbose, TEXT("FUdpMessageProcessor::ProcessDataSegment: Ignoring segment for recently delivered message %d"),
+			DataChunk.MessageId);
+		return;
+	}
+
 	TSharedPtr<FUdpReassembledMessage, ESPMode::ThreadSafe>& ReassembledMessage = NodeInfo.ReassembledMessages.FindOrAdd(DataChunk.MessageId);
 
 	// Reassemble message
@@ -852,10 +859,7 @@ void FUdpMessageProcessor::ProcessHelloSegment(FInboundSegment& Segment, FNodeIn
 	FGuid RemoteNodeId;
 	*Segment.Data << RemoteNodeId;
 
-	if (RemoteNodeId.IsValid())
-	{
-		NodeInfo.ResetIfRestarted(RemoteNodeId);
-	}
+	ensure(RemoteNodeId == NodeInfo.NodeId);
 }
 
 void FUdpMessageProcessor::ProcessPingSegment(FInboundSegment& Segment, FNodeInfo& NodeInfo)
@@ -867,10 +871,7 @@ void FUdpMessageProcessor::ProcessPingSegment(FInboundSegment& Segment, FNodeInf
 	uint8 NodeProtocolVersion;
 	*Segment.Data << NodeProtocolVersion;
 
-	if (RemoteNodeId.IsValid())
-	{
-		NodeInfo.ResetIfRestarted(RemoteNodeId);
-	}
+	ensure(RemoteNodeId == NodeInfo.NodeId);
 
 	// The protocol version we are going to use to communicate to this node is the smallest between its version and our own
 	uint8 ProtocolVersion = FMath::Min<uint8>(NodeProtocolVersion, UDP_MESSAGING_TRANSPORT_PROTOCOL_VERSION);
@@ -913,10 +914,7 @@ void FUdpMessageProcessor::ProcessPongSegment(FInboundSegment& Segment, FNodeInf
 	FGuid RemoteNodeId;
 	*Segment.Data << RemoteNodeId;
 
-	if (RemoteNodeId.IsValid())
-	{
-		NodeInfo.ResetIfRestarted(RemoteNodeId);
-	}
+	ensure(RemoteNodeId == NodeInfo.NodeId);
 }
 
 void FUdpMessageProcessor::ProcessRetransmitSegment(FInboundSegment& Segment, FNodeInfo& NodeInfo)
@@ -1590,6 +1588,8 @@ int32 FUdpMessageProcessor::UpdateSegmenters(FNodeInfo& NodeInfo)
 	const uint32 LostSegmentCount = NodeInfo.ExpireSegments(FUdpMessagingTime::Now());
 	NodeInfo.ComputeWindowSize(0, LostSegmentCount);
 
+	NodeInfo.ExpireRecentDeliveries();
+
 	NodeInfo.Statistics.PacketsInFlight = NodeInfo.InflightSegments.Num();
 	NodeInfo.Statistics.BytesInflight = NodeInfo.InflightSegments.Num() * NodeInfo.MaxSegmentSize;
 
@@ -1628,14 +1628,17 @@ bool FUdpMessageProcessor::UpdateReassemblers(FNodeInfo& NodeInfo)
 
 	for (TMap<int32, TSharedPtr<FUdpReassembledMessage, ESPMode::ThreadSafe>>::TIterator It(NodeInfo.ReassembledMessages); It; ++It)
 	{
+		const int32 MessageId = It.Key();
 		TSharedPtr<FUdpReassembledMessage, ESPMode::ThreadSafe>& ReassembledMessagePtr = It.Value();
 		// Believe it or not, TSharedPtr operator-> otherwise makes a big showing here.
 		FUdpReassembledMessage* ReassembledMessage = It.Value().Get();
 
 		if (ReassembledMessage->IsDelivered())
 		{
-			// If this is already delivered no need to do anything else. 
+			// Stop tracking the reassembler, but make a note to suppress double delivery of late retransmits.
 			It.RemoveCurrent();
+			NodeInfo.RecentlyDeliveredMessages.Add(MessageId);
+			NodeInfo.DeliveryExpiry.Emplace(FUdpMessagingTime::Now(), MessageId);
 			continue;
 		}
 		else if (ReassembledMessage->IsPendingDelivery())
@@ -1656,7 +1659,7 @@ bool FUdpMessageProcessor::UpdateReassemblers(FNodeInfo& NodeInfo)
 			TArray<uint32> PendingAcknowledgments = ReassembledMessage->TakePendingAcknowledgments(MaxAckNum, FUdpMessagingTime::Now());
 			const int32 AckCount = PendingAcknowledgments.Num();
 
-			FUdpMessageSegment::FAcknowledgeSegmentsChunk AcknowledgeChunk(It.Key()/*MessageId*/, MoveTemp(PendingAcknowledgments)/*Segments*/);
+			FUdpMessageSegment::FAcknowledgeSegmentsChunk AcknowledgeChunk(MessageId, MoveTemp(PendingAcknowledgments)/*Segments*/);
 			TSharedRef<FArrayWriter, ESPMode::ThreadSafe> Writer = MakeShared<FArrayWriter, ESPMode::ThreadSafe>();
 			{
 				*Writer << Header;
@@ -1664,20 +1667,16 @@ bool FUdpMessageProcessor::UpdateReassemblers(FNodeInfo& NodeInfo)
 			}
 			BytesSent += Writer->Num();
 
-			UE_LOG(LogUdpMessaging, Verbose, TEXT("FUdpMessageProcessor::UpdateReassemblers: Sending EUdpMessageSegments::AcknowledgeSegments for %d segments for message %d from %s"), AckCount, It.Key(), *ReassembledMessage->Describe());
+			UE_LOG(LogUdpMessaging, VeryVerbose, TEXT("FUdpMessageProcessor::UpdateReassemblers: Sending EUdpMessageSegments::AcknowledgeSegments for %d segments for message %d from %s"),
+				AckCount, MessageId, *ReassembledMessage->Describe());
 
 			const bool bAllowSegmentation_false = false;
 			const bool bHighPriority_true = true;
 			if (!SocketSender->Send(Writer, NodeInfo.Endpoint, bAllowSegmentation_false, bHighPriority_true))
 			{
 				UE_LOG(LogUdpMessaging, Error, TEXT("FUdpMessageProcessor::UpdateReassemblers: error sending EUdpMessageSegments::AcknowledgeSegments from %s"), *NodeInfo.NodeId.ToString());
-
 				return false;
 			}
-
-			UE_LOG(LogUdpMessaging, Verbose, TEXT("FUdpMessageProcessor::UpdateReassemblers: sending acknowledgement for reliable msg %d from %s"),
-				It.Key(),
-				*NodeInfo.NodeId.ToString());
 		}
 
 		// Try to deliver completed message that couldn't be delivered the first time around
@@ -1694,9 +1693,7 @@ bool FUdpMessageProcessor::UpdateReassemblers(FNodeInfo& NodeInfo)
 			{
 				const int ReceivedSegments = ReassembledMessage->GetTotalSegmentsCount() - ReassembledMessage->GetPendingSegmentsCount();
 				UE_LOG(LogUdpMessaging, Warning, TEXT("FUdpMessageProcessor::UpdateReassemblers Discarding %d/%d of stale message segments from %s"),
-					ReceivedSegments,
-					ReassembledMessage->GetTotalSegmentsCount(),
-					*ReassembledMessage->Describe());
+					ReceivedSegments, ReassembledMessage->GetTotalSegmentsCount(), *ReassembledMessage->Describe());
 			}
 
 			if (!ReassembledMessage->IsPendingDelivery())
diff --git a/Engine/Plugins/Messaging/UdpMessaging/Source/UdpMessaging/Private/Transport/UdpMessageProcessor.h b/Engine/Plugins/Messaging/UdpMessaging/Source/UdpMessaging/Private/Transport/UdpMessageProcessor.h
--- a/Engine/Plugins/Messaging/UdpMessaging/Source/UdpMessaging/Private/Transport/UdpMessageProcessor.h
+++ b/Engine/Plugins/Messaging/UdpMessaging/Source/UdpMessaging/Private/Transport/UdpMessageProcessor.h
@@ -67,13 +67,20 @@ struct FSentData
 };
 
 /** Lightweight tracking of segment transmit times to expedite maintenance of send window. */
-struct FSegmentExpiry
+struct FSegmentSendExpiry
 {
 	FUdpMessagingTime TimeSent;
 	int32 MessageId = 0;
 	uint32 SegmentNumber = 0;
 };
 
+/** Lightweight tracking of message delivery times to prevent double delivery. */
+struct FMessageDeliverExpiry
+{
+	FUdpMessagingTime TimeDelivered;
+	int32 MessageId = 0;
+};
+
 /** Returned by the segment processor. It provides details on what was sent and how it was sent. */
 struct FSentSegmentInfo
 {
@@ -173,10 +180,16 @@ class FUdpMessageProcessor
 		TMap<uint64, FSentData> InflightSegments;
 
 		/** Send timings for outgoing reliable segments. */
-		TRingBuffer<FSegmentExpiry> ReliableExpiry;
+		TRingBuffer<FSegmentSendExpiry> ReliableExpiry;
 
 		/** Send timings for outgoing unreliable segments. */
-		TRingBuffer<FSegmentExpiry> UnreliableExpiry;
+		TRingBuffer<FSegmentSendExpiry> UnreliableExpiry;
+
+		/** Messages delivered recently enough to track and suppress double delivery. */
+		TSet<int32> RecentlyDeliveredMessages;
+
+		/** Delivery times for recently delivered messages. */
+		TRingBuffer<FMessageDeliverExpiry> DeliveryExpiry;
 
 		/** Unique sequence id for every packet sent. */
 		uint64 SequenceId = 0;
@@ -330,8 +343,8 @@ class FUdpMessageProcessor
 			const uint64 Key = Remap(MessageId, SegmentNumber);
 			InflightSegments.Add(Key, SentData);
 
-			TRingBuffer<FSegmentExpiry>& ExpiryQueue = SentData.bIsReliable ? ReliableExpiry : UnreliableExpiry;
-			ExpiryQueue.Emplace(FSegmentExpiry{ .TimeSent = SentData.TimeSent,
+			TRingBuffer<FSegmentSendExpiry>& ExpiryQueue = SentData.bIsReliable ? ReliableExpiry : UnreliableExpiry;
+			ExpiryQueue.Emplace(FSegmentSendExpiry{ .TimeSent = SentData.TimeSent,
 				.MessageId = MessageId, .SegmentNumber = SegmentNumber });
 		}
 
@@ -426,12 +439,12 @@ class FUdpMessageProcessor
 			SCOPED_MESSAGING_TRACE(FNodeInfo_ExpireSegments);
 
 			auto ExpireUpTo = [this]
-			(TRingBuffer<FSegmentExpiry>& ExpiryQueue, FUdpMessagingTime ExpiryTime, bool bMarkLoss) -> uint32
+			(TRingBuffer<FSegmentSendExpiry>& ExpiryQueue, FUdpMessagingTime ExpiryTime, bool bMarkLoss) -> uint32
 			{
 				uint32 SegmentsExpired = 0;
 				while (ExpiryQueue.Num() && ExpiryQueue.First().TimeSent <= ExpiryTime)
 				{
-					const FSegmentExpiry Expiry = ExpiryQueue.PopFrontValue();
+					const FSegmentSendExpiry Expiry = ExpiryQueue.PopFrontValue();
 
 					const uint64 ExpiredKey = Remap(Expiry.MessageId, Expiry.SegmentNumber);
 					const FSentData* ExpiredData = InflightSegments.Find(ExpiredKey);
@@ -470,6 +483,18 @@ class FUdpMessageProcessor
 			return ReliableSegmentsLost;
 		}
 
+		/** Cull old redelivery guards. */
+		void ExpireRecentDeliveries()
+		{
+			const FTimespan RedeliveryGuardband = FTimespan::FromSeconds(10);
+			const FUdpMessagingTime ExpiryTime = FUdpMessagingTime::Now() - RedeliveryGuardband;
+			while (DeliveryExpiry.Num() && DeliveryExpiry.First().TimeDelivered <= ExpiryTime)
+			{
+				const FMessageDeliverExpiry Expiry = DeliveryExpiry.PopFrontValue();
+				RecentlyDeliveredMessages.Remove(Expiry.MessageId);
+			}
+		}
+
 		/**
 		 * Does this node have any segmenters that can send data for the current time frame.
 		 */
@@ -498,17 +523,6 @@ class FUdpMessageProcessor
 		{
 			return static_cast<int64>(WindowSize) - static_cast<int64>(InflightSegments.Num());
 		}
-		
-		/** Resets the endpoint info. */
-		void ResetIfRestarted(const FGuid& NewNodeId)
-		{
-			if (NewNodeId != NodeId)
-			{
-				ReassembledMessages.Reset();
-
-				NodeId = NewNodeId;
-			}
-		}
 	};
 
 
