UdpMessaging: Add MessagingInsights plugin for Unreal Insights, and corresponding instrumentation.

This plugin enables users to visualize Message Bus message summaries and lifecycles within Timing Insights. (Currently, only the UdpMessaging transport emits the required trace events, as this is by far the most commonly used transport.)

This changes FUdpMessagingTime to maintain integer cycles instead of double-precision seconds, which enables us to reuse those values directly for trace events.

Also removed several redundant uses of ESPMode::ThreadSafe, which is the default as of UE5.

#Synced-CL 49168713

diff --git a/Engine/Plugins/Messaging/MessagingInsights/MessagingInsights.uplugin b/Engine/Plugins/Messaging/MessagingInsights/MessagingInsights.uplugin
new file mode 100644
--- /dev/null
+++ b/Engine/Plugins/Messaging/MessagingInsights/MessagingInsights.uplugin
@@ -0,0 +1,34 @@
+{
+	"FileVersion": 3,
+	"Version": 1,
+	"VersionName": "1.0",
+	"FriendlyName": "Message Bus Insights",
+	"Description": "Allows debugging of Message Bus via Unreal Insights",
+	"Category": "Insights",
+	"CreatedBy": "Epic Games, Inc.",
+	"CreatedByURL": "https://epicgames.com",
+	"DocsURL": "",
+	"MarketplaceURL": "",
+	"SupportURL": "",
+	"EnabledByDefault": false,
+	"CanContainContent": false,
+	"IsBetaVersion": false,
+	"Installed": false,
+	"SupportedPrograms": [ "UnrealInsights" ],
+	"Plugins": [
+		{
+			"Name": "UdpMessaging",
+			"Enabled": true
+		}
+	],
+	"Modules": [
+		{
+			"Name": "UdpMessagingInsights",
+			"Type": "EditorAndProgram",
+			"LoadingPhase": "Default",
+			"ProgramAllowList": [
+				"UnrealInsights"
+			]
+		}
+	]
+}
diff --git a/Engine/Plugins/Messaging/MessagingInsights/Source/UdpMessagingInsights/Private/UdpMessagingInsightsModule.cpp b/Engine/Plugins/Messaging/MessagingInsights/Source/UdpMessagingInsights/Private/UdpMessagingInsightsModule.cpp
new file mode 100644
--- /dev/null
+++ b/Engine/Plugins/Messaging/MessagingInsights/Source/UdpMessagingInsights/Private/UdpMessagingInsightsModule.cpp
@@ -0,0 +1,39 @@
+// Copyright Epic Games, Inc. All Rights Reserved.
+
+#include "Features/IModularFeatures.h"
+#include "Modules/ModuleManager.h"
+#include "UdpMessagingTraceModule.h"
+#include "UdpMessagingTimingViewExtender.h"
+
+
+namespace UE::MessagingInsights
+{
+
+
+class FUdpMessagingInsightsModule : public IModuleInterface
+{
+public:
+	//~ Begin IModuleInterface interface
+	virtual void StartupModule() override
+	{
+		IModularFeatures::Get().RegisterModularFeature(TraceServices::ModuleFeatureName, &TraceModule);
+		IModularFeatures::Get().RegisterModularFeature(UE::Insights::Timing::TimingViewExtenderFeatureName, &TimingViewExtender);
+	}
+
+	virtual void ShutdownModule() override
+	{
+		IModularFeatures::Get().UnregisterModularFeature(TraceServices::ModuleFeatureName, &TraceModule);
+		IModularFeatures::Get().UnregisterModularFeature(UE::Insights::Timing::TimingViewExtenderFeatureName, &TimingViewExtender);
+	}
+	//~ End IModuleInterface interface
+
+private:
+	FUdpMessagingTraceModule TraceModule;
+	FUdpMessagingTimingViewExtender TimingViewExtender;
+};
+
+
+} // namespace UE::MessagingInsights
+
+
+IMPLEMENT_MODULE(UE::MessagingInsights::FUdpMessagingInsightsModule, UdpMessagingInsights);
diff --git a/Engine/Plugins/Messaging/MessagingInsights/Source/UdpMessagingInsights/Private/UdpMessagingTimingTrack.cpp b/Engine/Plugins/Messaging/MessagingInsights/Source/UdpMessagingInsights/Private/UdpMessagingTimingTrack.cpp
new file mode 100644
--- /dev/null
+++ b/Engine/Plugins/Messaging/MessagingInsights/Source/UdpMessagingInsights/Private/UdpMessagingTimingTrack.cpp
@@ -0,0 +1,380 @@
+// Copyright Epic Games, Inc. All Rights Reserved.
+
+#include "UdpMessagingTimingTrack.h"
+#include "Framework/MultiBox/MultiBoxBuilder.h"
+#include "HAL/PlatformApplicationMisc.h"
+#include "Insights/ViewModels/TimingTrackViewport.h"
+#include "Insights/ViewModels/TimingEvent.h"
+#include "Insights/ViewModels/TooltipDrawState.h"
+#include "InsightsCore/Common/TimeUtils.h"
+#include "Math/Color.h"
+#include "Math/UnitConversion.h"
+#include "Trace/UdpMessagingTrace.h"
+#include "UdpMessagingTraceProvider.h"
+#include "UdpMessagingTimingViewSession.h"
+
+
+#define LOCTEXT_NAMESPACE "FUdpMessagingTimingTrack"
+
+
+namespace UE::MessagingInsights
+{
+
+
+class FMessageTimingEvent : public FTimingEvent
+{
+	INSIGHTS_DECLARE_RTTI(FMessageTimingEvent, FTimingEvent)
+
+public:
+	FMessageTimingEvent(
+		const TSharedRef<const FBaseTimingTrack>& InTrack,
+		double InStartTime,
+		double InEndTime,
+		uint32 InDepth,
+		const FMessageLifecycleEvent& InEvent
+	)
+		: FTimingEvent(InTrack, InStartTime, InEndTime, InDepth)
+		, MessageEvent(InEvent)
+	{
+	}
+
+	FMessageLifecycleEvent MessageEvent;
+	const TCHAR* MessageTypePackage = nullptr;
+	const TCHAR* MessageTypeAsset = nullptr;
+	TOptional<FMessageSummary> MessageSummary;
+	TOptional<FNodeDiscoveredEvent> SenderNode;
+	TOptional<FNodeDiscoveredEvent> RecipientNode;
+};
+
+
+INSIGHTS_IMPLEMENT_RTTI(FMessageTimingEvent)
+INSIGHTS_IMPLEMENT_RTTI(FUdpMessagingTimingTrack)
+
+
+FUdpMessagingTimingTrack::FUdpMessagingTimingTrack(
+	const FUdpMessagingTimingViewSession& InSession,
+	const uint16 InNodeShortId,
+	const EMessageDirection InDirection
+)
+	: Super()
+	, SharedData(InSession)
+	, NodeShortId(InNodeShortId)
+	, Direction(InDirection)
+{
+}
+
+
+namespace Private
+{
+uint32 GetSpanColor(const FMessageLifecycleEvent& InEvent)
+{
+	using UE::UdpMessaging::Trace::EUdpMessageLifecycleEvent;
+	FLinearColor HSV(float(InEvent.MessageId) * UE_GOLDEN_RATIO * 360.f, 1.0f, 0.75f);
+
+	bool bIsInner = false;
+	int32 EventNum = 0;
+	if (InEvent.Details.IsType<FMessageLifecycleEvent::FInnerEvent>())
+	{
+		const FMessageLifecycleEvent::FInnerEvent& InnerDetails =
+			InEvent.Details.Get<FMessageLifecycleEvent::FInnerEvent>();
+
+		bIsInner = true;
+		EventNum = static_cast<int32>(InnerDetails.EventType);
+	}
+	else
+	{
+		const FMessageLifecycleEvent::FOuterEvent& OuterDetails =
+			InEvent.Details.Get<FMessageLifecycleEvent::FOuterEvent>();
+
+		EventNum = static_cast<int32>(OuterDetails.TerminalEventType);
+	}
+
+	if (EventNum >= static_cast<int32>(EUdpMessageLifecycleEvent::Discarded))
+	{
+		// Exceptional event; Insights idiom is to color it black.
+		HSV.B = 0.f;
+	}
+	else if (bIsInner)
+	{
+		if (EventNum >= static_cast<int32>(EUdpMessageLifecycleEvent::FirstSegmentEnqueueInbound))
+		{
+			// Receive event; make 0-relative
+			EventNum -= static_cast<int32>(EUdpMessageLifecycleEvent::FirstSegmentEnqueueInbound);
+		}
+
+		const float HueDegPerEvent = 30.f;
+		HSV.R += HueDegPerEvent * EventNum;
+		HSV.G *= 0.5f;
+	}
+
+	return HSV.HSVToLinearRGB().ToFColor(false).ToPackedARGB();
+}
+}
+
+
+void FUdpMessagingTimingTrack::BuildDrawState(
+	ITimingEventsTrackDrawStateBuilder& Builder,
+	const ITimingTrackUpdateContext& Context
+)
+{
+	const FUdpMessagingProvider* Provider =
+		SharedData.GetAnalysisSession().ReadProvider<FUdpMessagingProvider>(FUdpMessagingProvider::ProviderName);
+	if (!Provider)
+	{
+		return;
+	}
+
+	TraceServices::FAnalysisSessionReadScope SessionReadScope(SharedData.GetAnalysisSession());
+
+	const FTimingTrackViewport& Viewport = Context.GetViewport();
+	const FUdpMessagingProvider::FMessageTimeline* Timeline = Provider->GetMessageTimeline(NodeShortId, Direction);
+	if (!ensure(Timeline))
+	{
+		return;
+	}
+
+	auto BuildDrawEvents =
+		[this, &Builder, &Viewport, Provider]
+		(double InStartTime, double InEndTime, uint32 InDepth, const FMessageLifecycleEvent& MessageEvent)
+		{
+			TStringBuilder<128> NameBuilder;
+			if (MessageEvent.Details.IsType<FMessageLifecycleEvent::FInnerEvent>())
+			{
+				if (!bShowLifecycleEvents)
+				{
+					return TraceServices::EEventEnumerate::Continue;
+				}
+
+				const FMessageLifecycleEvent::FInnerEvent& InnerDetails =
+					MessageEvent.Details.Get<FMessageLifecycleEvent::FInnerEvent>();
+				NameBuilder.Appendf(TEXT("%s"), UE::UdpMessaging::Trace::LexToString(InnerDetails.EventType));
+			}
+			else
+			{
+				TPair<const TCHAR*, const TCHAR*> TypeInfo =
+					Provider->GetMessageTypeInfo(MessageEvent.SenderShortId, MessageEvent.RecipientShortId, MessageEvent.MessageId);
+
+				if (TypeInfo.Value)
+				{
+					NameBuilder.Appendf(TEXT("%s"), TypeInfo.Value);
+				}
+				else
+				{
+					NameBuilder.Appendf(TEXT("Message %u"), MessageEvent.MessageId);
+				}
+			}
+
+			const uint32 EffectiveLaneDepth = bShowLifecycleEvents ? MessageEvent.LaneDepth : MessageEvent.LaneDepth / 2;
+			Builder.AddEvent(InStartTime, InEndTime, EffectiveLaneDepth, *NameBuilder, 0, Private::GetSpanColor(MessageEvent));
+
+			return TraceServices::EEventEnumerate::Continue;
+		};
+
+	Timeline->EnumerateEvents(Viewport.GetStartTime(), Viewport.GetEndTime(), BuildDrawEvents);
+	Provider->EnumerateOutstandingMessages(NodeShortId, Direction, Viewport.GetStartTime(), Viewport.GetEndTime(), BuildDrawEvents);
+}
+
+const TSharedPtr<const ITimingEvent> FUdpMessagingTimingTrack::GetEvent(double InTime, double SecondsPerPixel, int32 InDepth) const
+{
+	const FUdpMessagingProvider* Provider =
+		SharedData.GetAnalysisSession().ReadProvider<FUdpMessagingProvider>(FUdpMessagingProvider::ProviderName);
+	if (!Provider)
+	{
+		return nullptr;
+	}
+
+	TraceServices::FAnalysisSessionReadScope SessionReadScope(SharedData.GetAnalysisSession());
+
+	double Delta = 2.0 * SecondsPerPixel;
+	FMessageLifecycleEvent BestMatchEvent;
+	double BestGraphStartTime = 0.0;
+	double BestGraphEndTime = 0.0;
+	uint32 BestDepth = 0;
+	auto SetBestMatchEvent =
+		[&BestMatchEvent, &BestGraphStartTime, &BestGraphEndTime, &BestDepth]
+		(double GraphStartTime, double GraphEndTime, uint32 Depth, const FMessageLifecycleEvent& MessageEvent)
+		{
+			BestMatchEvent = MessageEvent;
+			BestGraphStartTime = GraphStartTime;
+			BestGraphEndTime = GraphEndTime;
+			BestDepth = Depth;
+		};
+
+	const FUdpMessagingProvider::FMessageTimeline* Timeline = Provider->GetMessageTimeline(NodeShortId, Direction);
+	if (!ensure(Timeline))
+	{
+		return nullptr;
+	}
+
+	auto TryFindBetterMatchEvent =
+		[bShowLifecycleEvents = bShowLifecycleEvents, InTime, &Delta, &InDepth, &SetBestMatchEvent]
+		(double GraphStartTime, double GraphEndTime, uint32, const FMessageLifecycleEvent& MessageEvent)
+		{
+			const uint32 EffectiveInDepth = bShowLifecycleEvents ? InDepth : InDepth * 2;
+			if (GraphStartTime <= InTime && GraphEndTime >= InTime && EffectiveInDepth == MessageEvent.LaneDepth)
+			{
+				Delta = 0.0f;
+				SetBestMatchEvent(GraphStartTime, GraphEndTime, InDepth, MessageEvent);
+				return TraceServices::EEventEnumerate::Stop;
+			}
+
+			double DeltaLeft = InTime - GraphEndTime;
+			if (DeltaLeft >= 0 && DeltaLeft < Delta && EffectiveInDepth == MessageEvent.LaneDepth)
+			{
+				Delta = DeltaLeft;
+				SetBestMatchEvent(GraphStartTime, GraphEndTime, InDepth, MessageEvent);
+			}
+
+			double DeltaRight = GraphStartTime - InTime;
+			if (DeltaRight >= 0 && DeltaRight < Delta && EffectiveInDepth == MessageEvent.LaneDepth)
+			{
+				Delta = DeltaRight;
+				SetBestMatchEvent(GraphStartTime, GraphEndTime, InDepth, MessageEvent);
+			}
+
+			return TraceServices::EEventEnumerate::Continue;
+		};
+
+	Timeline->EnumerateEvents(InTime - Delta, InTime + Delta, TryFindBetterMatchEvent);
+	Provider->EnumerateOutstandingMessages(NodeShortId, Direction, InTime - Delta, InTime + Delta, TryFindBetterMatchEvent);
+
+	if (Delta < 2.0 * SecondsPerPixel)
+	{
+		TSharedRef<FMessageTimingEvent> Event = MakeShared<FMessageTimingEvent>(SharedThis(this), BestGraphStartTime, BestGraphEndTime, BestDepth, BestMatchEvent);
+
+		TPair<const TCHAR*, const TCHAR*> TypeInfo =
+			Provider->GetMessageTypeInfo(BestMatchEvent.SenderShortId, BestMatchEvent.RecipientShortId, BestMatchEvent.MessageId);
+		Event->MessageTypePackage = TypeInfo.Get<0>();
+		Event->MessageTypeAsset = TypeInfo.Get<1>();
+
+		if (const FMessageSummary* Summary = Provider->GetMessageSummary(BestMatchEvent.SenderShortId, BestMatchEvent.RecipientShortId, BestMatchEvent.MessageId))
+		{
+			Event->MessageSummary = *Summary;
+		}
+		if (const FNodeDiscoveredEvent* SenderNode = Provider->GetNodeByShortId(BestMatchEvent.SenderShortId))
+		{
+			Event->SenderNode = *SenderNode;
+		}
+		if (const FNodeDiscoveredEvent* RecipientNode = Provider->GetNodeByShortId(BestMatchEvent.RecipientShortId))
+		{
+			Event->RecipientNode = *RecipientNode;
+		}
+		return Event;
+	}
+
+	return nullptr;
+}
+
+void FUdpMessagingTimingTrack::InitTooltip(FTooltipDrawState& InOutTooltip, const ITimingEvent& InTooltipEvent) const
+{
+	InOutTooltip.ResetContent();
+
+	if (!InTooltipEvent.CheckTrack(this) || !InTooltipEvent.Is<FMessageTimingEvent>())
+	{
+		return;
+	}
+
+	const FMessageTimingEvent& TimingEvent = InTooltipEvent.As<FMessageTimingEvent>();
+
+	InOutTooltip.AddTitle(FString::Printf(TEXT("Message ID: %u"), TimingEvent.MessageEvent.MessageId));
+
+	EnumerateTimingEventFields(TimingEvent,
+		[&InOutTooltip](FStringView FieldName, FStringView FieldValue)
+		{
+			InOutTooltip.AddNameValueTextLine(FieldName, FieldValue);
+		});
+
+	InOutTooltip.UpdateLayout();
+}
+
+void FUdpMessagingTimingTrack::OnClipboardCopyEvent(const ITimingEvent& InSelectedEvent) const
+{
+	if (!InSelectedEvent.CheckTrack(this) || !InSelectedEvent.Is<FMessageTimingEvent>())
+	{
+		return;
+	}
+
+	const FMessageTimingEvent& TimingEvent = InSelectedEvent.As<FMessageTimingEvent>();
+
+	TStringBuilder<2048> Builder;
+	Builder.Appendf(TEXT("Message ID: %u\n"), TimingEvent.MessageEvent.MessageId);
+
+	EnumerateTimingEventFields(TimingEvent,
+		[&Builder](FStringView FieldName, FStringView FieldValue)
+		{
+			Builder.Appendf(TEXT("%.*s %.*s\n"),
+				FieldName.Len(), FieldName.GetData(),
+				FieldValue.Len(), FieldValue.GetData());
+		});
+
+	FPlatformApplicationMisc::ClipboardCopy(*Builder);
+}
+
+void FUdpMessagingTimingTrack::EnumerateTimingEventFields(const FMessageTimingEvent& InTimingEvent, TFunctionRef<void(FStringView, FStringView)> InCallback) const
+{
+	TStringBuilder<256> Builder;
+
+	InCallback(TEXTVIEW("Start Time:"), UE::Insights::FormatTimeAuto(InTimingEvent.GetStartTime(), 6));
+	InCallback(TEXTVIEW("End Time:"), UE::Insights::FormatTimeAuto(InTimingEvent.GetEndTime(), 6));
+	InCallback(TEXTVIEW("Duration:"), UE::Insights::FormatTimeAuto(InTimingEvent.GetDuration()));
+	if (InTimingEvent.MessageEvent.Details.IsType<FMessageLifecycleEvent::FInnerEvent>())
+	{
+		const FMessageLifecycleEvent::FInnerEvent& InnerDetails =
+			InTimingEvent.MessageEvent.Details.Get<FMessageLifecycleEvent::FInnerEvent>();
+		InCallback(TEXTVIEW("Event:"), UE::UdpMessaging::Trace::LexToString(InnerDetails.EventType));
+	}
+	else
+	{
+		InCallback(TEXT("Type Package:"), InTimingEvent.MessageTypePackage);
+		InCallback(TEXT("Type Asset:"), InTimingEvent.MessageTypeAsset);
+		if (InTimingEvent.MessageSummary)
+		{
+			Builder.Reset();
+			const uint64 SizeBytes = InTimingEvent.MessageSummary->Size;
+			const FNumericUnit<uint64> SizeUnit = FUnitConversion::QuantizeUnitsToBestFit(SizeBytes, EUnit::Bytes);
+			Builder.Appendf(TEXT("%llu %s"), SizeUnit.Value, FUnitConversion::GetUnitDisplayString(SizeUnit.Units));
+			Builder.Appendf(TEXT(" (%s bytes)"), *FText::AsNumber(SizeBytes).ToString());
+			InCallback(TEXTVIEW("Size:"), Builder.ToView());
+			
+			Builder.Reset();
+			Builder.Appendf(TEXT("%u"), InTimingEvent.MessageSummary->NumSegments);
+			InCallback(TEXTVIEW("# Segments:"), Builder.ToView());
+		}
+
+		if (InTimingEvent.SenderNode)
+		{
+			InCallback(TEXTVIEW("Sender:"), InTimingEvent.SenderNode->Endpoint.ToString());
+		}
+
+		if (InTimingEvent.RecipientNode)
+		{
+			InCallback(TEXTVIEW("Recipient:"), InTimingEvent.RecipientNode->Endpoint.ToString());
+		}
+	}
+}
+
+void FUdpMessagingTimingTrack::BuildContextMenu(FMenuBuilder& InOutMenuBuilder)
+{
+	InOutMenuBuilder.BeginSection("UdpMessaging", LOCTEXT("ContextMenu_SectionHeading", "UdpMessaging"));
+
+	InOutMenuBuilder.AddMenuEntry(
+		LOCTEXT("ContextMenu_LifecycleEvents_Label", "Lifecycle Events"), 
+		LOCTEXT("ContextMenu_LifecycleEvents_Tooltip", "Show/Hide nested individual lifecycle events for messages"),
+		FSlateIcon(),
+		FUIAction(
+			FExecuteAction::CreateSPLambda(this, [this]() { bShowLifecycleEvents = !bShowLifecycleEvents;}),
+			FCanExecuteAction(),
+			FIsActionChecked::CreateSPLambda(this, [this]() { return bShowLifecycleEvents; })
+		),
+		NAME_None,
+		EUserInterfaceActionType::Check
+	);
+
+	InOutMenuBuilder.EndSection();
+}
+
+
+} // namespace UE::MessagingInsights
+
+
+#undef LOCTEXT_NAMESPACE
diff --git a/Engine/Plugins/Messaging/MessagingInsights/Source/UdpMessagingInsights/Private/UdpMessagingTimingTrack.h b/Engine/Plugins/Messaging/MessagingInsights/Source/UdpMessagingInsights/Private/UdpMessagingTimingTrack.h
new file mode 100644
--- /dev/null
+++ b/Engine/Plugins/Messaging/MessagingInsights/Source/UdpMessagingInsights/Private/UdpMessagingTimingTrack.h
@@ -0,0 +1,52 @@
+// Copyright Epic Games, Inc. All Rights Reserved.
+
+#pragma once
+
+#include "Insights/ViewModels/TimingEventsTrack.h"
+
+class ITimingTrackUpdateContext;
+
+
+namespace UE::MessagingInsights
+{
+
+
+enum class EMessageDirection : uint8;
+class FMessageTimingEvent;
+class FUdpMessagingTimingViewSession;
+
+
+/** Timing track for UDP Messaging. */
+class FUdpMessagingTimingTrack : public FTimingEventsTrack
+{
+	INSIGHTS_DECLARE_RTTI(FUdpMessagingTimingTrack, FTimingEventsTrack)
+	using Super = FTimingEventsTrack;
+
+public:
+	FUdpMessagingTimingTrack(
+		const FUdpMessagingTimingViewSession& InSharedData,
+		const uint16 InNodeShortId,
+		const EMessageDirection InDirection
+	);
+
+	//~ Begin FTimingEventsTrack interface
+	virtual void BuildDrawState(ITimingEventsTrackDrawStateBuilder& Builder, const ITimingTrackUpdateContext& Context) override;
+	virtual void InitTooltip(FTooltipDrawState& InOutTooltip, const ITimingEvent& InTooltipEvent) const override;
+	virtual void OnClipboardCopyEvent(const ITimingEvent& InSelectedEvent) const override;
+	virtual void BuildContextMenu(FMenuBuilder& MenuBuilder) override;
+	virtual const TSharedPtr<const ITimingEvent> GetEvent(double InTime, double SecondsPerPixel, int32 Depth) const override;
+	//~ End FTimingEventsTrack interface
+
+private:
+	void EnumerateTimingEventFields(const FMessageTimingEvent& InSelectedEvent, TFunctionRef<void(FStringView, FStringView)> InCallback) const;
+
+private:
+	const FUdpMessagingTimingViewSession& SharedData;
+	const uint16 NodeShortId;
+	const EMessageDirection Direction;
+
+	bool bShowLifecycleEvents = true;
+};
+
+
+} //namespace UE::MessagingInsights
diff --git a/Engine/Plugins/Messaging/MessagingInsights/Source/UdpMessagingInsights/Private/UdpMessagingTimingViewExtender.cpp b/Engine/Plugins/Messaging/MessagingInsights/Source/UdpMessagingInsights/Private/UdpMessagingTimingViewExtender.cpp
new file mode 100644
--- /dev/null
+++ b/Engine/Plugins/Messaging/MessagingInsights/Source/UdpMessagingInsights/Private/UdpMessagingTimingViewExtender.cpp
@@ -0,0 +1,72 @@
+// Copyright Epic Games, Inc. All Rights Reserved.
+
+#include "UdpMessagingTimingViewExtender.h"
+#include "UdpMessagingTimingViewSession.h"
+#include "Insights/ITimingViewSession.h"
+#include "Insights/IUnrealInsightsModule.h"
+
+
+namespace UE::MessagingInsights
+{
+
+
+void FUdpMessagingTimingViewExtender::OnBeginSession(UE::Insights::Timing::ITimingViewSession& InSession)
+{
+	if (InSession.GetName() != FInsightsManagerTabs::TimingProfilerTabId)
+	{
+		return;
+	}
+
+	FPerSessionData* PerSessionData = PerSessionDataMap.Find(&InSession);
+	if (!PerSessionData)
+	{
+		PerSessionData = &PerSessionDataMap.Add(&InSession);
+		PerSessionData->SharedData = MakePimpl<FUdpMessagingTimingViewSession>();
+	}
+
+	PerSessionData->SharedData->OnBeginSession(InSession);
+}
+
+void FUdpMessagingTimingViewExtender::OnEndSession(UE::Insights::Timing::ITimingViewSession& InSession)
+{
+	if (InSession.GetName() != FInsightsManagerTabs::TimingProfilerTabId)
+	{
+		return;
+	}
+
+	if (FPerSessionData* PerSessionData = PerSessionDataMap.Find(&InSession))
+	{
+		PerSessionData->SharedData->OnEndSession(InSession);
+	}
+
+	PerSessionDataMap.Remove(&InSession);
+}
+
+void FUdpMessagingTimingViewExtender::Tick(const UE::Insights::Timing::FTimingViewExtenderTickParams& InParams)
+{
+	if (InParams.Session.GetName() != FInsightsManagerTabs::TimingProfilerTabId)
+	{
+		return;
+	}
+
+	if (FPerSessionData* PerSessionData = PerSessionDataMap.Find(&InParams.Session))
+	{
+		PerSessionData->SharedData->Tick(InParams);
+	}
+}
+
+void FUdpMessagingTimingViewExtender::ExtendFilterMenu(UE::Insights::Timing::ITimingViewSession& InSession, FMenuBuilder& InMenuBuilder)
+{
+	if (InSession.GetName() != FInsightsManagerTabs::TimingProfilerTabId)
+	{
+		return;
+	}
+
+	if (FPerSessionData* PerSessionData = PerSessionDataMap.Find(&InSession))
+	{
+		PerSessionData->SharedData->ExtendFilterMenu(InMenuBuilder);
+	}
+}
+
+
+} //namespace UE::MessagingInsights
diff --git a/Engine/Plugins/Messaging/MessagingInsights/Source/UdpMessagingInsights/Private/UdpMessagingTimingViewExtender.h b/Engine/Plugins/Messaging/MessagingInsights/Source/UdpMessagingInsights/Private/UdpMessagingTimingViewExtender.h
new file mode 100644
--- /dev/null
+++ b/Engine/Plugins/Messaging/MessagingInsights/Source/UdpMessagingInsights/Private/UdpMessagingTimingViewExtender.h
@@ -0,0 +1,42 @@
+// Copyright Epic Games, Inc. All Rights Reserved.
+
+#pragma once
+
+#include "Insights/ITimingViewExtender.h"
+#include "Templates/PimplPtr.h"
+
+
+namespace TraceServices { class IAnalysisSession; }
+namespace UE::Insights::Timing { class ITimingViewSession; }
+namespace UE::Insights::Timing { struct FTimingViewExtenderTickParams; }
+class FMenuBuilder;
+
+
+namespace UE::MessagingInsights
+{
+
+class FUdpMessagingTimingViewSession;
+
+
+class FUdpMessagingTimingViewExtender : public UE::Insights::Timing::ITimingViewExtender
+{
+public:
+	//~ Begin UE::Insights::Timing::ITimingViewExtender interface
+	virtual void OnBeginSession(UE::Insights::Timing::ITimingViewSession& InSession) override;
+	virtual void OnEndSession(UE::Insights::Timing::ITimingViewSession& InSession) override;
+	virtual void Tick(const UE::Insights::Timing::FTimingViewExtenderTickParams& InParams) override;
+	virtual void ExtendFilterMenu(UE::Insights::Timing::ITimingViewSession& InSession, FMenuBuilder& InMenuBuilder) override;
+	//~ End UE::Insights::Timing::ITimingViewExtender interface
+
+private:
+	struct FPerSessionData
+	{
+		TPimplPtr<FUdpMessagingTimingViewSession> SharedData;
+	};
+
+	// The data we host per-session
+	TMap<UE::Insights::Timing::ITimingViewSession*, FPerSessionData> PerSessionDataMap;
+};
+
+
+} //namespace UE::MessagingInsights
diff --git a/Engine/Plugins/Messaging/MessagingInsights/Source/UdpMessagingInsights/Private/UdpMessagingTimingViewSession.cpp b/Engine/Plugins/Messaging/MessagingInsights/Source/UdpMessagingInsights/Private/UdpMessagingTimingViewSession.cpp
new file mode 100644
--- /dev/null
+++ b/Engine/Plugins/Messaging/MessagingInsights/Source/UdpMessagingInsights/Private/UdpMessagingTimingViewSession.cpp
@@ -0,0 +1,141 @@
+// Copyright Epic Games, Inc. All Rights Reserved.
+
+#include "UdpMessagingTimingViewSession.h"
+
+#include "Insights/ITimingViewExtender.h"
+#include "Insights/ITimingViewSession.h"
+#include "Insights/IUnrealInsightsModule.h"
+
+#include "UdpMessagingTimingTrack.h"
+#include "UdpMessagingTraceProvider.h"
+
+#include "Framework/MultiBox/MultiBoxBuilder.h"
+
+
+#define LOCTEXT_NAMESPACE "UdpMessagingTimingViewSession"
+
+
+namespace UE::MessagingInsights
+{
+
+
+FUdpMessagingTimingViewSession::FUdpMessagingTimingViewSession()
+{
+}
+
+void FUdpMessagingTimingViewSession::OnBeginSession(UE::Insights::Timing::ITimingViewSession& InTimingViewSession)
+{
+	if (InTimingViewSession.GetName() == FInsightsManagerTabs::TimingProfilerTabId)
+	{
+		TimingViewSession = &InTimingViewSession;
+	}
+}
+
+void FUdpMessagingTimingViewSession::OnEndSession(UE::Insights::Timing::ITimingViewSession& InTimingViewSession)
+{
+	if (&InTimingViewSession != TimingViewSession)
+	{
+		return;
+	}
+
+	LocalNodeTracks.Empty();
+
+	TimingViewSession = nullptr;
+}
+
+void FUdpMessagingTimingViewSession::Tick(const UE::Insights::Timing::FTimingViewExtenderTickParams& InParams)
+{
+	if (&InParams.Session != TimingViewSession)
+	{
+		return;
+	}
+
+	if (!InParams.AnalysisSession)
+	{
+		return;
+	}
+
+	AnalysisSession = InParams.AnalysisSession;
+
+	if (const FUdpMessagingProvider* Provider =
+		AnalysisSession->ReadProvider<FUdpMessagingProvider>(FUdpMessagingProvider::ProviderName))
+	{
+		TraceServices::FAnalysisSessionReadScope SessionReadScope(GetAnalysisSession());
+
+		// Add timing tracks
+		for (uint32 LocalNodeIdx = 0; LocalNodeIdx < Provider->GetNumLocalNodes(); ++LocalNodeIdx)
+		{
+			const uint16 LocalNodeShortId = Provider->GetLocalNodeShortId(LocalNodeIdx);
+			const FUdpMessagingProvider::FMessageTimeline* SendTimeline =
+				Provider->GetMessageTimeline(LocalNodeShortId, EMessageDirection::Send);
+			const FUdpMessagingProvider::FMessageTimeline* ReceiveTimeline =
+				Provider->GetMessageTimeline(LocalNodeShortId, EMessageDirection::Receive);
+			if (!ensure(SendTimeline && ReceiveTimeline))
+			{
+				continue;
+			}
+
+			if (SendTimeline->GetEventCount() == 0 && ReceiveTimeline->GetEventCount() == 0)
+			{
+				continue;
+			}
+
+			if (!LocalNodeTracks.Contains(LocalNodeShortId))
+			{
+				TSharedRef<FUdpMessagingTimingTrack> SendTrack = MakeShared<FUdpMessagingTimingTrack>(*this, LocalNodeShortId, EMessageDirection::Send);
+				TSharedRef<FUdpMessagingTimingTrack> ReceiveTrack = MakeShared<FUdpMessagingTimingTrack>(*this, LocalNodeShortId, EMessageDirection::Receive);
+				LocalNodeTracks.Emplace(LocalNodeShortId, { .SendTrack = SendTrack, .ReceiveTrack = ReceiveTrack });
+
+				if (const FNodeDiscoveredEvent* NodeInfo = Provider->GetNodeByShortId(LocalNodeShortId))
+				{
+					SendTrack->SetName(FString::Printf(TEXT("UdpMessaging - %s - Sent"), *NodeInfo->Endpoint.ToString()));
+					ReceiveTrack->SetName(FString::Printf(TEXT("UdpMessaging - %s - Received"), *NodeInfo->Endpoint.ToString()));
+				}
+
+				SendTrack->SetVisibilityFlag(bMessageTracksEnabled);
+				ReceiveTrack->SetVisibilityFlag(bMessageTracksEnabled);
+				TimingViewSession->AddScrollableTrack(SendTrack);
+				TimingViewSession->AddScrollableTrack(ReceiveTrack);
+				TimingViewSession->InvalidateScrollableTracksOrder();
+			}
+		}
+	}
+}
+
+void FUdpMessagingTimingViewSession::ExtendFilterMenu(FMenuBuilder& InMenuBuilder)
+{
+	InMenuBuilder.BeginSection("UdpMessagingTracks", LOCTEXT("UdpMessagingHeader", "UdpMessaging"));
+	{
+		InMenuBuilder.AddMenuEntry(
+			LOCTEXT("MessageTrackToggleLabel", "Messages"),
+			LOCTEXT("MessageTrackToggleLabel_Tooltip", "Show/hide the Messages tracks"),
+			FSlateIcon(),
+			FUIAction(
+				FExecuteAction::CreateRaw(this, &FUdpMessagingTimingViewSession::ToggleMessageTracks),
+				FCanExecuteAction(),
+				FIsActionChecked::CreateLambda([this]() { return bMessageTracksEnabled; })
+			),
+			NAME_None,
+			EUserInterfaceActionType::ToggleButton
+		);
+	}
+	InMenuBuilder.EndSection();
+}
+
+
+void FUdpMessagingTimingViewSession::ToggleMessageTracks()
+{
+	bMessageTracksEnabled = !bMessageTracksEnabled;
+
+	for (TPair<uint16, FDuplexTracks>& Node : LocalNodeTracks)
+	{
+		Node.Value.SendTrack->SetVisibilityFlag(bMessageTracksEnabled);
+		Node.Value.ReceiveTrack->SetVisibilityFlag(bMessageTracksEnabled);
+	}
+}
+
+
+} //namespace UE::MessagingInsights
+
+
+#undef LOCTEXT_NAMESPACE
diff --git a/Engine/Plugins/Messaging/MessagingInsights/Source/UdpMessagingInsights/Private/UdpMessagingTimingViewSession.h b/Engine/Plugins/Messaging/MessagingInsights/Source/UdpMessagingInsights/Private/UdpMessagingTimingViewSession.h
new file mode 100644
--- /dev/null
+++ b/Engine/Plugins/Messaging/MessagingInsights/Source/UdpMessagingInsights/Private/UdpMessagingTimingViewSession.h
@@ -0,0 +1,65 @@
+// Copyright Epic Games, Inc. All Rights Reserved.
+
+#pragma once
+
+#include "Templates/SharedPointer.h"
+
+namespace TraceServices { class IAnalysisSession; }
+namespace UE::Insights::Timing { enum class ETimeChangedFlags : int32; }
+namespace UE::Insights::Timing { struct FTimingViewExtenderTickParams; }
+namespace UE::Insights::Timing { class ITimingViewSession; }
+namespace UE::MessagingInsights { class FUdpMessagingTimingTrack; }
+class FMenuBuilder;
+
+
+namespace UE::MessagingInsights
+{
+
+
+class FUdpMessagingTimingViewSession
+{
+public:
+	FUdpMessagingTimingViewSession();
+
+	void OnBeginSession(UE::Insights::Timing::ITimingViewSession& InTimingViewSession);
+	void OnEndSession(UE::Insights::Timing::ITimingViewSession& InTimingViewSession);
+	void Tick(const UE::Insights::Timing::FTimingViewExtenderTickParams& InParams);
+	void ExtendFilterMenu(FMenuBuilder& InMenuBuilder);
+
+	/** Get the last cached analysis session */
+	const TraceServices::IAnalysisSession& GetAnalysisSession() const
+	{
+		check(AnalysisSession);
+		return *AnalysisSession;
+	}
+
+	/** Show/hide the message tracks. */
+	void ToggleMessageTracks();
+
+	/** The timing view for the session. */
+	UE::Insights::Timing::ITimingViewSession* GetTimingView() const
+	{
+		return TimingViewSession;
+	}
+
+private:
+	// Cached analysis session, set in Tick()
+	const TraceServices::IAnalysisSession* AnalysisSession = nullptr;
+
+	// Cached timing view session, set in OnBeginSession/OnEndSession
+	UE::Insights::Timing::ITimingViewSession* TimingViewSession = nullptr;
+
+	struct FDuplexTracks
+	{
+		TSharedRef<FUdpMessagingTimingTrack> SendTrack;
+		TSharedRef<FUdpMessagingTimingTrack> ReceiveTrack;
+	};
+
+	TMap<uint16, FDuplexTracks> LocalNodeTracks;
+
+	/** Whether message tracks are enabled. */
+	bool bMessageTracksEnabled = true;
+};
+
+
+} //namespace UE::MessagingInsights
diff --git a/Engine/Plugins/Messaging/MessagingInsights/Source/UdpMessagingInsights/Private/UdpMessagingTraceAnalyzer.cpp b/Engine/Plugins/Messaging/MessagingInsights/Source/UdpMessagingInsights/Private/UdpMessagingTraceAnalyzer.cpp
new file mode 100644
--- /dev/null
+++ b/Engine/Plugins/Messaging/MessagingInsights/Source/UdpMessagingInsights/Private/UdpMessagingTraceAnalyzer.cpp
@@ -0,0 +1,90 @@
+// Copyright Epic Games, Inc. All Rights Reserved.
+
+#include "UdpMessagingTraceAnalyzer.h"
+#include "HAL/LowLevelMemTracker.h"
+#include "TraceServices/Model/Definitions.h"
+#include "TraceServices/Model/Strings.h"
+#include "UdpMessagingTraceProvider.h"
+
+
+namespace UE::MessagingInsights
+{
+
+
+FUdpMessagingAnalyzer::FUdpMessagingAnalyzer(TraceServices::IAnalysisSession& InSession, FUdpMessagingProvider& InProvider)
+	: Session(InSession)
+	, Provider(InProvider)
+{
+
+}
+
+
+void FUdpMessagingAnalyzer::OnAnalysisBegin(const FOnAnalysisContext& Context)
+{
+	UE::Trace::IAnalyzer::FInterfaceBuilder& Builder = Context.InterfaceBuilder;
+
+	Builder.RouteEvent(RouteId_DiscoveredNode, "UdpMessaging", "DiscoveredNode");
+	Builder.RouteEvent(RouteId_MessageSummary, "UdpMessaging", "MessageSummary");
+	Builder.RouteEvent(RouteId_MessageTypeInfo, "UdpMessaging", "MessageTypeInfo");
+	Builder.RouteEvent(RouteId_MessageLifecycle, "UdpMessaging", "MessageLifecycle");
+}
+
+
+void FUdpMessagingAnalyzer::OnAnalysisEnd()
+{
+}
+
+
+bool FUdpMessagingAnalyzer::OnEvent(uint16 RouteId, EStyle Style, const FOnEventContext& Context)
+{
+	LLM_SCOPE_BYNAME(TEXT("Insights/FUdpMessagingAnalyzer"));
+
+	TraceServices::FAnalysisSessionEditScope _(Session);
+
+	const UE::Trace::IAnalyzer::FEventData& EventData = Context.EventData;
+	switch (RouteId)
+	{
+		case RouteId_DiscoveredNode:
+		{
+			FNodeDiscoveredEvent Event = FNodeDiscoveredEvent::FromEventData(EventData);
+			Provider.AddNodeDiscoveredEvent(Context.EventTime.AsSeconds(Event.Cycle), MoveTemp(Event));
+			break;
+		}
+		case RouteId_MessageSummary:
+		{
+			FMessageSummary Summary = FMessageSummary::FromEventData(EventData);
+			Provider.AddMessageSummary(Summary.SenderShortId, Summary.RecipientShortId, Summary.MessageId, MoveTemp(Summary));
+			break;
+		}
+		case RouteId_MessageTypeInfo:
+		{
+			const TraceServices::IDefinitionProvider* DefinitionProvider = TraceServices::ReadDefinitionProvider(Session);
+			if (DefinitionProvider)
+			{
+				const uint16 SenderShortId = EventData.GetValue<uint16>("SenderShortId");
+				const uint16 RecipientShortId = EventData.GetValue<uint16>("RecipientShortId");
+				const uint32 MessageId = EventData.GetValue<uint32>("MessageId");
+				const UE::Trace::FEventRef32 PackageNameRef = EventData.GetReferenceValue<uint32>("Package");
+				const UE::Trace::FEventRef32 AssetNameRef = EventData.GetReferenceValue<uint32>("Asset");
+				const TraceServices::FStringDefinition* PackageName = DefinitionProvider->Get<TraceServices::FStringDefinition>(PackageNameRef);
+				const TraceServices::FStringDefinition* AssetName = DefinitionProvider->Get<TraceServices::FStringDefinition>(AssetNameRef);
+				if (PackageName && AssetName)
+				{
+					Provider.AddMessageTypeInfo(SenderShortId, RecipientShortId, MessageId, PackageName->Display, AssetName->Display);
+				}
+			}
+			break;
+		}
+		case RouteId_MessageLifecycle:
+		{
+			FMessageLifecycleEvent Event = FMessageLifecycleEvent::FromEventData(EventData);
+			Provider.AddMessageEvent(Context.EventTime.AsSeconds(Event.Cycle), MoveTemp(Event));
+			break;
+		}
+	}
+
+	return true;
+}
+
+
+} // namespace UE::MessagingInsights
diff --git a/Engine/Plugins/Messaging/MessagingInsights/Source/UdpMessagingInsights/Private/UdpMessagingTraceAnalyzer.h b/Engine/Plugins/Messaging/MessagingInsights/Source/UdpMessagingInsights/Private/UdpMessagingTraceAnalyzer.h
new file mode 100644
--- /dev/null
+++ b/Engine/Plugins/Messaging/MessagingInsights/Source/UdpMessagingInsights/Private/UdpMessagingTraceAnalyzer.h
@@ -0,0 +1,45 @@
+// Copyright Epic Games, Inc. All Rights Reserved.
+
+#pragma once
+
+#include "Trace/Analyzer.h"
+
+namespace TraceServices
+{
+	class IAnalysisSession;
+}
+
+
+namespace UE::MessagingInsights
+{
+
+
+class FUdpMessagingProvider;
+
+
+class FUdpMessagingAnalyzer : public UE::Trace::IAnalyzer
+{
+public:
+	FUdpMessagingAnalyzer(TraceServices::IAnalysisSession& InSession, FUdpMessagingProvider& InProvider);
+
+	//~ Begin IAnalyzer interface
+	virtual void OnAnalysisBegin(const FOnAnalysisContext& Context) override;
+	virtual void OnAnalysisEnd() override;
+	virtual bool OnEvent(uint16 RouteId, EStyle Style, const FOnEventContext& Context) override;
+	//~ End IAnalyzer interface
+
+private:
+	enum : uint16
+	{
+		RouteId_DiscoveredNode,
+		RouteId_MessageSummary,
+		RouteId_MessageTypeInfo,
+		RouteId_MessageLifecycle,
+	};
+
+	TraceServices::IAnalysisSession& Session;
+	FUdpMessagingProvider& Provider;
+};
+
+
+} // namespace UE::MessagingInsights
diff --git a/Engine/Plugins/Messaging/MessagingInsights/Source/UdpMessagingInsights/Private/UdpMessagingTraceModule.cpp b/Engine/Plugins/Messaging/MessagingInsights/Source/UdpMessagingInsights/Private/UdpMessagingTraceModule.cpp
new file mode 100644
--- /dev/null
+++ b/Engine/Plugins/Messaging/MessagingInsights/Source/UdpMessagingInsights/Private/UdpMessagingTraceModule.cpp
@@ -0,0 +1,25 @@
+// Copyright Epic Games, Inc. All Rights Reserved.
+
+#include "UdpMessagingTraceModule.h"
+#include "UdpMessagingTraceAnalyzer.h"
+#include "UdpMessagingTraceProvider.h"
+
+namespace UE::MessagingInsights
+{
+
+FName FUdpMessagingTraceModule::ModuleName("TraceModule_UdpMessaging");
+
+void FUdpMessagingTraceModule::GetModuleInfo(TraceServices::FModuleInfo& OutModuleInfo)
+{
+	OutModuleInfo.Name = ModuleName;
+	OutModuleInfo.DisplayName = TEXT("UdpMessaging");
+}
+
+void FUdpMessagingTraceModule::OnAnalysisBegin(TraceServices::IAnalysisSession& InSession)
+{
+	TSharedRef<FUdpMessagingProvider> Provider = MakeShared<FUdpMessagingProvider>(InSession);
+	InSession.AddProvider(FUdpMessagingProvider::ProviderName, Provider);
+	InSession.AddAnalyzer(new FUdpMessagingAnalyzer(InSession, *Provider));
+}
+
+} // namespace UE::MessagingInsights
diff --git a/Engine/Plugins/Messaging/MessagingInsights/Source/UdpMessagingInsights/Private/UdpMessagingTraceModule.h b/Engine/Plugins/Messaging/MessagingInsights/Source/UdpMessagingInsights/Private/UdpMessagingTraceModule.h
new file mode 100644
--- /dev/null
+++ b/Engine/Plugins/Messaging/MessagingInsights/Source/UdpMessagingInsights/Private/UdpMessagingTraceModule.h
@@ -0,0 +1,22 @@
+// Copyright Epic Games, Inc. All Rights Reserved.
+
+#pragma once
+
+#include "TraceServices/ModuleService.h"
+
+
+namespace UE::MessagingInsights
+{
+
+class FUdpMessagingTraceModule : public TraceServices::IModule
+{
+	static FName ModuleName;
+
+public:
+	//~ Begin TraceServices::IModule interface
+	virtual void GetModuleInfo(TraceServices::FModuleInfo& OutModuleInfo) override;
+	virtual void OnAnalysisBegin(TraceServices::IAnalysisSession& Session) override;
+	//~ End TraceServices::IModule interface
+};
+
+} // namespace UE::MessagingInsights
diff --git a/Engine/Plugins/Messaging/MessagingInsights/Source/UdpMessagingInsights/Private/UdpMessagingTraceProvider.cpp b/Engine/Plugins/Messaging/MessagingInsights/Source/UdpMessagingInsights/Private/UdpMessagingTraceProvider.cpp
new file mode 100644
--- /dev/null
+++ b/Engine/Plugins/Messaging/MessagingInsights/Source/UdpMessagingInsights/Private/UdpMessagingTraceProvider.cpp
@@ -0,0 +1,329 @@
+// Copyright Epic Games, Inc. All Rights Reserved.
+
+#include "UdpMessagingTraceProvider.h"
+#include "Trace/UdpMessagingTrace.h"
+
+
+namespace UE::MessagingInsights
+{
+
+using namespace UE::UdpMessaging::Trace;
+
+
+const FName FUdpMessagingProvider::ProviderName("UdpMessagingProvider");
+
+
+FUdpMessagingProvider::FUdpMessagingProvider(TraceServices::IAnalysisSession& InSession)
+	: Session(InSession)
+	, NodeDiscoveredTimeline(Session.GetLinearAllocator())
+{
+}
+
+
+void FUdpMessagingProvider::AddNodeDiscoveredEvent(double InEventTimeSeconds, FNodeDiscoveredEvent&& InEvent)
+{
+	Session.WriteAccessCheck();
+
+	const bool bIsLocalNode = InEvent.IsLocalNode();
+	const uint16 ShortId = InEvent.ShortId;
+	const uint64 EventIdx = NodeDiscoveredTimeline.EmplaceEvent(InEventTimeSeconds, MoveTemp(InEvent));
+
+	if (ensure(!NodeShortIdToDiscoveredEvent.Contains(ShortId)))
+	{
+		NodeShortIdToDiscoveredEvent.Add(ShortId, EventIdx);
+
+		const FNodeDiscoveredEvent& Event = NodeDiscoveredTimeline.GetEvent(EventIdx);
+		if (bIsLocalNode)
+		{
+			LocalNodeShortIds.Add(ShortId);
+			LocalNodes.Add(ShortId, {
+				.SendTimeline{ .MessageTimeline = FMessageTimeline(Session.GetLinearAllocator()) },
+				.ReceiveTimeline{ .MessageTimeline = FMessageTimeline(Session.GetLinearAllocator()) },
+			});
+		}
+	}
+}
+
+
+const FNodeDiscoveredEvent* FUdpMessagingProvider::GetNodeByShortId(uint16 InShortId) const
+{
+	Session.ReadAccessCheck();
+
+	if (const uint64* DiscoveredEvent = NodeShortIdToDiscoveredEvent.Find(InShortId))
+	{
+		return &NodeDiscoveredTimeline.GetEvent(*DiscoveredEvent);
+	}
+
+	return nullptr;
+}
+
+
+namespace Private
+{
+	bool IsInitialEvent(UE::UdpMessaging::Trace::EUdpMessageLifecycleEvent InEvent)
+	{
+		switch (InEvent)
+		{
+			case UE::UdpMessaging::Trace::EUdpMessageLifecycleEvent::EnqueueOutbound: // Send
+			case UE::UdpMessaging::Trace::EUdpMessageLifecycleEvent::FirstSegmentEnqueueInbound: // Receive
+				return true;
+			default:
+				return false;
+		}
+	}
+
+	bool IsTerminalEvent(UE::UdpMessaging::Trace::EUdpMessageLifecycleEvent InEvent)
+	{
+		switch (InEvent)
+		{
+			case UE::UdpMessaging::Trace::EUdpMessageLifecycleEvent::Complete: // Send
+			case UE::UdpMessaging::Trace::EUdpMessageLifecycleEvent::Delivered: // Receive
+			case UE::UdpMessaging::Trace::EUdpMessageLifecycleEvent::Discarded:
+				return true;
+			default:
+				return false;
+		}
+	}
+}
+
+
+void FUdpMessagingProvider::AddMessageEvent(double InEventTimeSeconds, FMessageLifecycleEvent&& InEvent)
+{
+	Session.WriteAccessCheck();
+
+	check(InEvent.Details.IsType<FMessageLifecycleEvent::FInnerEvent>());
+
+	if (FLocalNode* LocalSender = LocalNodes.Find(InEvent.SenderShortId))
+	{
+		LocalSender->SendTimeline.AddMessageEvent(InEventTimeSeconds, MoveTemp(InEvent));
+	}
+	else if (FLocalNode* LocalReceiver = LocalNodes.Find(InEvent.RecipientShortId))
+	{
+		LocalReceiver->ReceiveTimeline.AddMessageEvent(InEventTimeSeconds, MoveTemp(InEvent));
+	}
+	else
+	{
+		ensure(false);
+	}
+}
+
+// See the comment about deferred commit in FLocalNodeTimeline::AddMessageEvent below.
+void FUdpMessagingProvider::EnumerateOutstandingMessages(
+	uint16 InLocalNodeShortId, EMessageDirection InDirection,
+	double InIntervalStart, double InIntervalEnd,
+	TraceServices::ITimeline<FMessageLifecycleEvent>::EventRangeCallback InCallback
+) const
+{
+	Session.ReadAccessCheck();
+
+	const FLocalNode* LocalNode = LocalNodes.Find(InLocalNodeShortId);
+	if (!LocalNode)
+	{
+		return;
+	}
+
+	const FLocalNodeTimeline& LocalTimeline = (InDirection == EMessageDirection::Send)
+		? LocalNode->SendTimeline : LocalNode->ReceiveTimeline;
+
+	TArray<FEventTiming> SortedEventTimings;
+	int32 NextUnassignedLane = LocalTimeline.LaneFreeTimes.Num() * 2;
+	for (const TPair<FMessageKey, TArray<FEventTiming>>& OutstandingMessage : LocalTimeline.OutstandingMessageEvents)
+	{
+		SortedEventTimings = OutstandingMessage.Value;
+		SortedEventTimings.StableSort([](const FEventTiming& Lhs, const FEventTiming& Rhs) { return Lhs.Value.Cycle < Rhs.Value.Cycle; });
+
+		const double MessageStartTime = SortedEventTimings[0].Key;
+		if (MessageStartTime >= InIntervalEnd)
+		{
+			continue;
+		}
+
+		const int32 OuterLaneDepth = NextUnassignedLane;
+		NextUnassignedLane += 2;
+
+		FMessageLifecycleEvent TempEvent(SortedEventTimings[0].Value);
+		TempEvent.LaneDepth = OuterLaneDepth;
+		TempEvent.Details.Set<FMessageLifecycleEvent::FOuterEvent>({
+			SortedEventTimings.Last().Value.Details.Get<FMessageLifecycleEvent::FInnerEvent>().EventType });
+
+		InCallback(MessageStartTime, std::numeric_limits<double>::infinity(), 0, TempEvent);
+
+		for (int32 EventIdx = 0; EventIdx < SortedEventTimings.Num(); ++EventIdx)
+		{
+			const FEventTiming& EventTiming = SortedEventTimings[EventIdx];
+			const double StartTime = EventTiming.Key;
+			const double EndTime = (EventIdx + 1) < SortedEventTimings.Num()
+				? SortedEventTimings[EventIdx + 1].Key
+				: std::numeric_limits<double>::infinity();
+
+			TempEvent = EventTiming.Value;
+			TempEvent.LaneDepth = OuterLaneDepth + 1;
+			InCallback(StartTime, EndTime, 0, TempEvent);
+		}
+	}
+}
+
+
+void FUdpMessagingProvider::FLocalNodeTimeline::AddMessageEvent(double InEventTimeSeconds, FMessageLifecycleEvent&& InEvent)
+{
+	const UE::UdpMessaging::Trace::EUdpMessageLifecycleEvent EventType =
+		InEvent.Details.Get<FMessageLifecycleEvent::FInnerEvent>().EventType;
+
+	const bool bIsTerminalEvent = Private::IsTerminalEvent(EventType);
+	const FMessageKey MessageKey = InEvent.GetMessageKey();
+	TArray<FEventTiming>& EventTimings = OutstandingMessageEvents.FindOrAdd(MessageKey);
+	EventTimings.Emplace(InEventTimeSeconds, MoveTemp(InEvent));
+
+	// Wait until we've buffered all the events for this message, then commit them in a second pass.
+	// Deferring the lane allocation lets us handle events emitted with a retroactive `Cycles` timestamp,
+	// or otherwise analyzed out of order, without inadvertently generating overlapping intervals.
+	// This means enumerating unterminated messages requires special handling; see EnumerateOutstandingMessages.
+	if (bIsTerminalEvent)
+	{
+		EventTimings.StableSort([](const FEventTiming& Lhs, const FEventTiming& Rhs) { return Lhs.Value.Cycle < Rhs.Value.Cycle; });
+
+		const double StartTime = EventTimings[0].Key;
+		const double EndTime = EventTimings.Last().Key;
+		const int32 OuterLaneDepth = AllocLaneDepth(StartTime, EndTime);
+
+		// Initialize the outer event as a copy of the first inner event.
+		FMessageLifecycleEvent OuterEvent(EventTimings[0].Value);
+		OuterEvent.Details.Set<FMessageLifecycleEvent::FOuterEvent>({EventType});
+		OuterEvent.LaneDepth = OuterLaneDepth;
+
+		const uint64 OuterEventIdx = MessageTimeline.EmplaceBeginEvent(StartTime, MoveTemp(OuterEvent));
+		TOptional<uint64> PrevEventIdx;
+
+		for (FEventTiming& EventTiming : EventTimings)
+		{
+			if (PrevEventIdx.IsSet())
+			{
+				MessageTimeline.EndEvent(PrevEventIdx.GetValue(), EventTiming.Key);
+			}
+
+			if (!Private::IsTerminalEvent(EventTiming.Value.Details.Get<FMessageLifecycleEvent::FInnerEvent>().EventType))
+			{
+				EventTiming.Value.LaneDepth = OuterLaneDepth + 1;
+				PrevEventIdx = MessageTimeline.EmplaceBeginEvent(EventTiming.Key, MoveTemp(EventTiming.Value));
+			}
+		}
+
+		MessageTimeline.EndEvent(OuterEventIdx, EndTime);
+
+		OutstandingMessageEvents.Remove(MessageKey); // Note: Invalidates EventTimings
+	}
+}
+
+
+int32 FUdpMessagingProvider::FLocalNodeTimeline::AllocLaneDepth(double InStartTime, double InEndTime)
+{
+	int32 LaneIdx = LaneFreeTimes.IndexOfByPredicate([InStartTime](double Elem) { return Elem < InStartTime; });
+	if (LaneIdx == INDEX_NONE)
+	{
+		LaneIdx = LaneFreeTimes.Add(InEndTime);
+	}
+	else
+	{
+		LaneFreeTimes[LaneIdx] = InEndTime;
+	}
+
+	return LaneIdx * 2;
+}
+
+
+void FUdpMessagingProvider::AddMessageSummary(uint16 InSenderShortId, uint16 InRecipientShortId, uint32 InMessageId, FMessageSummary&& InMessageSummary)
+{
+	Session.WriteAccessCheck();
+	const FMessageKey MessageKey = MakeMessageKey(InSenderShortId, InRecipientShortId, InMessageId);
+	if (ensure(!MessageKeyToSummary.Contains(MessageKey)))
+	{
+		MessageKeyToSummary.Emplace(MessageKey, MoveTemp(InMessageSummary));
+	}
+}
+
+
+const FMessageSummary* FUdpMessagingProvider::GetMessageSummary(uint16 InSenderShortId, uint16 InRecipientShortId, uint32 InMessageId) const
+{
+	Session.ReadAccessCheck();
+	const FMessageKey MessageKey = MakeMessageKey(InSenderShortId, InRecipientShortId, InMessageId);
+	return MessageKeyToSummary.Find(MessageKey);
+}
+
+
+void FUdpMessagingProvider::AddMessageTypeInfo(uint16 InSenderShortId, uint16 InRecipientShortId, uint32 InMessageId, const TCHAR* InPackage, const TCHAR* InAsset)
+{
+	Session.WriteAccessCheck();
+	const FMessageKey MessageKey = MakeMessageKey(InSenderShortId, InRecipientShortId, InMessageId);
+	if (ensure(!MessageKeyToTypeInfo.Contains(MessageKey)))
+	{
+		MessageKeyToTypeInfo.Add(MessageKey, { InPackage, InAsset });
+	}
+}
+
+
+TPair<const TCHAR*, const TCHAR*> FUdpMessagingProvider::GetMessageTypeInfo(uint16 SenderShortId, uint16 RecipientShortId, uint32 MessageId) const
+{
+	Session.ReadAccessCheck();
+	const FMessageKey MessageKey = MakeMessageKey(SenderShortId, RecipientShortId, MessageId);
+	if (const TPair<const TCHAR*, const TCHAR*>* TypeInfo = MessageKeyToTypeInfo.Find(MessageKey))
+	{
+		return *TypeInfo;
+	}
+
+	return { nullptr, nullptr };
+}
+
+
+FNodeDiscoveredEvent FNodeDiscoveredEvent::FromEventData(const UE::Trace::IAnalyzer::FEventData& EventData)
+{
+	return {
+		.Cycle = EventData.GetValue<uint64>("Cycle"),
+		.ShortId = EventData.GetValue<uint16>("ShortId"),
+		.DiscoveredByShortId = EventData.GetValue<uint16>("DiscoveredByShortId"),
+		.Id = FGuid{
+			EventData.GetValue<uint32>("IdA"),
+			EventData.GetValue<uint32>("IdB"),
+			EventData.GetValue<uint32>("IdC"),
+			EventData.GetValue<uint32>("IdD"),
+		},
+		.Endpoint = {
+			FIPv4Address(
+				EventData.GetValue<uint8>("IpAddrA"),
+				EventData.GetValue<uint8>("IpAddrB"),
+				EventData.GetValue<uint8>("IpAddrC"),
+				EventData.GetValue<uint8>("IpAddrD")
+			),
+			EventData.GetValue<uint16>("Port")
+		},
+		.ProtocolVersion = EventData.GetValue<uint8>("ProtocolVersion"),
+	};
+}
+
+FMessageSummary FMessageSummary::FromEventData(const UE::Trace::IAnalyzer::FEventData& EventData)
+{
+	return {
+		.SenderShortId = EventData.GetValue<uint16>("SenderShortId"),
+		.RecipientShortId = EventData.GetValue<uint16>("RecipientShortId"),
+		.MessageId = EventData.GetValue<uint32>("MessageId"),
+		.Size = EventData.GetValue<uint64>("Size"),
+		.NumSegments = EventData.GetValue<uint32>("NumSegments"),
+	};
+}
+
+
+FMessageLifecycleEvent FMessageLifecycleEvent::FromEventData(const UE::Trace::IAnalyzer::FEventData& EventData)
+{
+	return {
+		.Cycle = EventData.GetValue<uint64>("Cycle"),
+		.SenderShortId = EventData.GetValue<uint16>("SenderShortId"),
+		.RecipientShortId = EventData.GetValue<uint16>("RecipientShortId"),
+		.MessageId = EventData.GetValue<uint32>("MessageId"),
+		.Details = TVariant<FOuterEvent, FInnerEvent>(TInPlaceType<FInnerEvent>(),
+			static_cast<EUdpMessageLifecycleEvent>(
+				EventData.GetValue<std::underlying_type_t<EUdpMessageLifecycleEvent>>("LifecycleEvent"))
+		),
+	};
+}
+
+
+} // namespace UE::MessagingInsights
diff --git a/Engine/Plugins/Messaging/MessagingInsights/Source/UdpMessagingInsights/Private/UdpMessagingTraceProvider.h b/Engine/Plugins/Messaging/MessagingInsights/Source/UdpMessagingInsights/Private/UdpMessagingTraceProvider.h
new file mode 100644
--- /dev/null
+++ b/Engine/Plugins/Messaging/MessagingInsights/Source/UdpMessagingInsights/Private/UdpMessagingTraceProvider.h
@@ -0,0 +1,193 @@
+// Copyright Epic Games, Inc. All Rights Reserved.
+
+#pragma once
+
+#include "Interfaces/IPv4/IPv4Endpoint.h"
+#include "Model/IntervalTimeline.h"
+#include "Model/PointTimeline.h"
+#include "Trace/Analyzer.h"
+#include "TraceServices/Model/AnalysisSession.h"
+
+namespace UE::UdpMessaging::Trace
+{
+	enum class EUdpMessageLifecycleEvent : uint8;
+}
+
+
+namespace UE::MessagingInsights
+{
+using FMessageKey = uint64;
+
+inline FMessageKey MakeMessageKey(uint16 SenderShortId, uint16 RecipientShortId, uint32 MessageId)
+{
+	return
+		uint64(SenderShortId)    << 48 |
+		uint64(RecipientShortId) << 32 |
+		MessageId;
+}
+
+struct FNodeDiscoveredEvent
+{
+	uint64 Cycle;
+	uint16 ShortId;
+	uint16 DiscoveredByShortId;
+	FGuid Id;
+	FIPv4Endpoint Endpoint;
+	uint8 ProtocolVersion;
+
+	static FNodeDiscoveredEvent FromEventData(const UE::Trace::IAnalyzer::FEventData& EventData);
+
+	bool IsLocalNode() const
+	{
+		return DiscoveredByShortId == 0;
+	}
+};
+
+struct FMessageSummary
+{
+	uint16 SenderShortId;
+	uint16 RecipientShortId;
+	uint32 MessageId;
+	uint64 Size;
+	uint32 NumSegments;
+
+	static FMessageSummary FromEventData(const UE::Trace::IAnalyzer::FEventData& EventData);
+};
+
+struct FMessageLifecycleEvent
+{
+	uint64 Cycle;
+	uint16 SenderShortId;
+	uint16 RecipientShortId;
+	uint32 MessageId;
+
+	struct FOuterEvent
+	{
+		UE::UdpMessaging::Trace::EUdpMessageLifecycleEvent TerminalEventType;
+	};
+
+	struct FInnerEvent
+	{
+		UE::UdpMessaging::Trace::EUdpMessageLifecycleEvent EventType;
+	};
+
+	TVariant<FOuterEvent, FInnerEvent> Details;
+
+	// Computed/derived values
+	int32 LaneDepth = 0;
+
+	static FMessageLifecycleEvent FromEventData(const UE::Trace::IAnalyzer::FEventData& EventData);
+
+	FMessageKey GetMessageKey() const
+	{
+		return MakeMessageKey(SenderShortId, RecipientShortId, MessageId);
+	}
+};
+
+
+enum class EMessageDirection : uint8
+{
+	Send,
+	Receive,
+};
+
+
+class FUdpMessagingProvider : public TraceServices::IProvider
+{
+	using FEventTiming = TPair<double, FMessageLifecycleEvent>;
+
+public:
+	static const FName ProviderName;
+
+	FUdpMessagingProvider(TraceServices::IAnalysisSession& InSession);
+
+	using FNodeDiscoveredTimeline = TraceServices::TPointTimeline<FNodeDiscoveredEvent>;
+	const FNodeDiscoveredTimeline& GetNodeDiscoveredTimeline() const
+	{
+		Session.ReadAccessCheck();
+		return NodeDiscoveredTimeline;
+	}
+	void AddNodeDiscoveredEvent(double InEventTimeSeconds, FNodeDiscoveredEvent&& InEvent);
+
+	using FMessageTimeline = TraceServices::TIntervalTimeline<FMessageLifecycleEvent>;
+	const FMessageTimeline* GetMessageTimeline(uint16 InLocalNodeShortId, EMessageDirection InDirection) const
+	{
+		Session.ReadAccessCheck();
+		if (const FLocalNode* LocalNode = LocalNodes.Find(InLocalNodeShortId); ensure(LocalNode))
+		{
+			switch (InDirection)
+			{
+				case EMessageDirection::Send: return &LocalNode->SendTimeline.MessageTimeline;
+				case EMessageDirection::Receive: return &LocalNode->ReceiveTimeline.MessageTimeline;
+				default: checkNoEntry();
+			}
+		}
+
+		return nullptr;
+	}
+
+	void EnumerateOutstandingMessages(
+		uint16 InLocalNodeShortId, EMessageDirection InDirection,
+		double InIntervalStart, double InIntervalEnd,
+		TraceServices::ITimeline<FMessageLifecycleEvent>::EventRangeCallback InCallback
+	) const;
+
+	uint32 GetNumLocalNodes() const
+	{
+		Session.ReadAccessCheck();
+		return LocalNodeShortIds.Num();
+	}
+
+	uint16 GetLocalNodeShortId(uint32 InLocalNodeIdx) const
+	{
+		Session.ReadAccessCheck();
+		check(InLocalNodeIdx < (uint32)LocalNodeShortIds.Num());
+		return LocalNodeShortIds[InLocalNodeIdx];
+	}
+
+	const FNodeDiscoveredEvent* GetNodeByShortId(uint16 InShortId) const;
+
+	void AddMessageEvent(double InEventTimeSeconds, FMessageLifecycleEvent&& InEvent);
+
+	void AddMessageSummary(uint16 InSenderShortId, uint16 InRecipientShortId, uint32 InMessageId, FMessageSummary&& InMessageSummary);
+	const FMessageSummary* GetMessageSummary(uint16 InSenderShortId, uint16 InRecipientShortId, uint32 InMessageId) const;
+
+	void AddMessageTypeInfo(uint16 InSenderShortId, uint16 InRecipientShortId, uint32 InMessageId, const TCHAR* InPackage, const TCHAR* InAsset);
+	TPair<const TCHAR*, const TCHAR*> GetMessageTypeInfo(uint16 InSenderShortId, uint16 InRecipientShortId, uint32 InMessageId) const;
+
+private:
+	TraceServices::IAnalysisSession& Session;
+
+	FNodeDiscoveredTimeline NodeDiscoveredTimeline;
+	TMap<uint16, uint64> NodeShortIdToDiscoveredEvent;
+
+	/** Contains the short IDs just for local nodes (typically just a few). */
+	TArray<uint16> LocalNodeShortIds;
+
+	struct FLocalNodeTimeline
+	{
+		FMessageTimeline MessageTimeline;
+		TArray<double> LaneFreeTimes;
+
+		TMap<FMessageKey, TArray<FEventTiming>> OutstandingMessageEvents;
+
+		void AddMessageEvent(double InEventTimeSeconds, FMessageLifecycleEvent&& InEvent);
+		int32 AllocLaneDepth(double InStartTime, double InEndTime);
+	};
+
+	struct FLocalNode
+	{
+		FLocalNodeTimeline SendTimeline;
+		FLocalNodeTimeline ReceiveTimeline;
+	};
+
+	TMap<uint16, FLocalNode> LocalNodes;
+
+	// Pointers are to strings owned by TraceServices::FStringDefinition, guaranteed valid for the analysis session.
+	TMap<FMessageKey, TPair<const TCHAR*, const TCHAR*>> MessageKeyToTypeInfo;
+
+	TMap<FMessageKey, FMessageSummary> MessageKeyToSummary;
+};
+
+
+} // namespace UE::MessagingInsights
diff --git a/Engine/Plugins/Messaging/MessagingInsights/Source/UdpMessagingInsights/UdpMessagingInsights.Build.cs b/Engine/Plugins/Messaging/MessagingInsights/Source/UdpMessagingInsights/UdpMessagingInsights.Build.cs
new file mode 100644
--- /dev/null
+++ b/Engine/Plugins/Messaging/MessagingInsights/Source/UdpMessagingInsights/UdpMessagingInsights.Build.cs
@@ -0,0 +1,35 @@
+// Copyright Epic Games, Inc. All Rights Reserved.
+
+namespace UnrealBuildTool.Rules
+{
+	public class UdpMessagingInsights : ModuleRules
+	{
+		public UdpMessagingInsights(ReadOnlyTargetRules Target) : base(Target)
+		{
+			PrivateDependencyModuleNames.AddRange(new string[]
+			{
+				"ApplicationCore",
+				"Core",
+				"CoreUObject",
+				"Messaging",
+				"Networking",
+				"Slate",
+				"SlateCore",
+				"TraceAnalysis",
+				"TraceLog",
+				"TraceInsights",
+				"TraceInsightsCore",
+				"TraceServices",
+				"UdpMessaging",
+			});
+
+			if (Target.bCompileAgainstEngine)
+			{
+				PrivateDependencyModuleNames.AddRange(new string[]
+				{
+					"Engine",
+				});
+			}
+		}
+	}
+}
diff --git a/Engine/Plugins/Messaging/UdpMessaging/Source/UdpMessaging/Internal/Trace/UdpMessagingTrace.h b/Engine/Plugins/Messaging/UdpMessaging/Source/UdpMessaging/Internal/Trace/UdpMessagingTrace.h
new file mode 100644
--- /dev/null
+++ b/Engine/Plugins/Messaging/UdpMessaging/Source/UdpMessaging/Internal/Trace/UdpMessagingTrace.h
@@ -0,0 +1,84 @@
+// Copyright Epic Games, Inc. All Rights Reserved.
+
+#pragma once
+
+#include "MessagingTrace.h"
+#include "Trace/Trace.h"
+
+
+enum class EMessageFlags : uint32;
+class FName;
+struct FGuid;
+struct FIPv4Endpoint;
+struct FTopLevelAssetPath;
+template<typename T, typename... Ts> class TVariant;
+
+
+#if !defined(UDPMESSAGINGTRACE_ENABLED)
+#	define UDPMESSAGINGTRACE_ENABLED MESSAGINGTRACE_ENABLED
+#endif
+
+
+namespace UE::UdpMessaging::Trace
+{
+	enum class EUdpMessageLifecycleEvent : uint8
+	{
+		// Send lifecycle
+		EnqueueOutbound,
+		BeginSerialize,
+		EndSerialize,
+		ConsumeOutbound,
+		Complete,
+
+		// Receive lifecycle
+		FirstSegmentEnqueueInbound,
+		BeginReassembly,
+		EndReassembly,
+		PendingDelivery,
+		BeginDeserialize,
+		Delivered,
+
+		// Exceptional (send or receive)
+		Discarded,
+	};
+
+	UDPMESSAGING_API const TCHAR* LexToString(EUdpMessageLifecycleEvent InEvent);
+}
+
+
+#if UDPMESSAGINGTRACE_ENABLED
+
+
+namespace UE::UdpMessaging::Trace
+{
+	void OutputDiscoveredNode(uint16 ShortId, uint16 DiscoveredByShortId, const FGuid& NodeId, const FIPv4Endpoint& Endpoint, uint8 ProtocolVersion);
+
+	void OutputMessageSummary(uint16 SenderShortId, uint16 RecipientShortId, uint32 MessageId, uint64 Size, uint32 NumSegments, EMessageFlags Flags);
+	void OutputMessageTypeInfo(uint16 SenderShortId, uint16 RecipientShortId, uint32 MessageId, const TVariant<FName, FTopLevelAssetPath>& TypeInfo);
+	void OutputMessageLifecycle(uint64 Cycle, uint16 SenderShortId, uint16 RecipientShortId, uint32 MessageId, EUdpMessageLifecycleEvent Event);
+}
+
+
+#define TRACE_UDPMESSAGING_DISCOVERED_NODE(ShortId, DiscoveredByShortId, NodeId, Endpoint, ProtocolVersion) \
+	UE::UdpMessaging::Trace::OutputDiscoveredNode(ShortId, DiscoveredByShortId, NodeId, Endpoint, ProtocolVersion);
+
+#define TRACE_UDPMESSAGING_MESSAGE_SUMMARY(SenderShortId, RecipientShortId, MessageId, Size, NumSegments, Flags) \
+	UE::UdpMessaging::Trace::OutputMessageSummary(SenderShortId, RecipientShortId, MessageId, Size, NumSegments, Flags);
+
+#define TRACE_UDPMESSAGING_MESSAGE_TYPEINFO(SenderShortId, RecipientShortId, MessageId, TypeInfo) \
+	UE::UdpMessaging::Trace::OutputMessageTypeInfo(SenderShortId, RecipientShortId, MessageId, TypeInfo);
+
+#define TRACE_UDPMESSAGING_MESSAGE_LIFECYCLE(Cycle, SenderShortId, RecipientShortId, MessageId, Event) \
+	UE::UdpMessaging::Trace::OutputMessageLifecycle(Cycle, SenderShortId, RecipientShortId, MessageId, Event);
+
+
+#else // #if UDPMESSAGINGTRACE_ENABLED
+
+
+#define TRACE_UDPMESSAGING_DISCOVERED_NODE(...)
+#define TRACE_UDPMESSAGING_MESSAGE_SUMMARY(...)
+#define TRACE_UDPMESSAGING_MESSAGE_TYPEINFO(...)
+#define TRACE_UDPMESSAGING_MESSAGE_LIFECYCLE(...)
+
+
+#endif // #if UDPMESSAGINGTRACE_ENABLED .. #else
diff --git a/Engine/Plugins/Messaging/UdpMessaging/Source/UdpMessaging/Private/Shared/SocketSender.cpp b/Engine/Plugins/Messaging/UdpMessaging/Source/UdpMessaging/Private/Shared/SocketSender.cpp
--- a/Engine/Plugins/Messaging/UdpMessaging/Source/UdpMessaging/Private/Shared/SocketSender.cpp
+++ b/Engine/Plugins/Messaging/UdpMessaging/Source/UdpMessaging/Private/Shared/SocketSender.cpp
@@ -8,8 +8,8 @@
 #include "Logging/StructuredLog.h"
 #include "Sockets.h"
 #include "SocketSubsystem.h"
+#include "Trace/UdpMessagingTrace.h"
 #include "UdpMessagingPrivate.h"
-#include "UdpMessagingTracing.h"
 
 
 namespace UE::UdpMessaging {
@@ -18,6 +18,7 @@ namespace UE::UdpMessaging {
 FSocketSender::FSocketSender(FSocket* InSocket, const TCHAR* InDescription, const FOptions& InOptions /* = FOptions() */)
 	: Socket(InSocket)
 	, SocketSubsystem(ISocketSubsystem::Get(PLATFORM_SOCKETSUBSYSTEM))
+	, SocketAddr(SocketSubsystem->CreateInternetAddr())
 	, bStopping(false)
 	, WaitTime(InOptions.WaitTime)
 	, Description(InDescription)
@@ -25,6 +26,8 @@ FSocketSender::FSocketSender(FSocket* InSocket, const TCHAR* InDescription, cons
 	check(Socket != nullptr);
 	check(Socket->GetSocketType() == SOCKTYPE_Datagram);
 
+	Socket->GetAddress(*SocketAddr);
+
 	int32 NewSize = 0;
 	Socket->SetSendBufferSize(InOptions.SendBufferSize, NewSize);
 
@@ -60,7 +63,7 @@ bool FSocketSender::Send(
 	const TSharedRef<TArray<uint8>>& InData,
 	const FIPv4Endpoint& InRecipient,
 	bool bAllowSegmentation, /* = false */
-    bool bHighPriority /* = false */
+	bool bHighPriority /* = false */
 )
 {
 	if (!bStopping)
diff --git a/Engine/Plugins/Messaging/UdpMessaging/Source/UdpMessaging/Private/Shared/SocketSender.h b/Engine/Plugins/Messaging/UdpMessaging/Source/UdpMessaging/Private/Shared/SocketSender.h
--- a/Engine/Plugins/Messaging/UdpMessaging/Source/UdpMessaging/Private/Shared/SocketSender.h
+++ b/Engine/Plugins/Messaging/UdpMessaging/Source/UdpMessaging/Private/Shared/SocketSender.h
@@ -11,6 +11,7 @@
 
 
 class FEvent;
+class FInternetAddr;
 class FRunnableThread;
 class FSocket;
 class ISocketSubsystem;
@@ -101,6 +102,9 @@ public:
 		bool bHighPriority = false
 	);
 
+	/** Get the address to which this socket is bound. */
+	const TSharedRef<FInternetAddr>& GetAddress() const { return SocketAddr; }
+
 	/** If send segmentation offload is enabled, returns the max datagram payload. Otherwise 0. */
 	uint16 GetSegmentOffloadSize() const { return SegmentOffloadSize; }
 
@@ -148,6 +152,9 @@ private:
 	/** Cached pointer to the platform socket subsystem. */
 	ISocketSubsystem* SocketSubsystem;
 
+	/** Cached socket binding address. */
+	TSharedRef<FInternetAddr> SocketAddr;
+
 	/** Flag indicating that the thread is stopping. */
 	std::atomic<bool> bStopping;
 
diff --git a/Engine/Plugins/Messaging/UdpMessaging/Source/UdpMessaging/Private/Shared/UdpMessagingTime.h b/Engine/Plugins/Messaging/UdpMessaging/Source/UdpMessaging/Private/Shared/UdpMessagingTime.h
--- a/Engine/Plugins/Messaging/UdpMessaging/Source/UdpMessaging/Private/Shared/UdpMessagingTime.h
+++ b/Engine/Plugins/Messaging/UdpMessaging/Source/UdpMessaging/Private/Shared/UdpMessagingTime.h
@@ -3,22 +3,30 @@
 #pragma once
 
 #include "HAL/PlatformTime.h"
+#include "Misc/DateTime.h"
 #include "Misc/Timespan.h"
 
 
-/** Thin wrapper to help migrate from FDateTime to FPlatformTime::Seconds for relative measurements. */
+/** Thin wrapper to help migrate from FDateTime to FPlatformTime for relative measurements. */
 struct FUdpMessagingTime
 {
-	double RawValue = std::numeric_limits<double>::lowest();
+	int64 Cycles = 0;
 
 	static FUdpMessagingTime Now()
 	{
-		return { .RawValue = FPlatformTime::Seconds() };
+		return { .Cycles = static_cast<int64>(FPlatformTime::Cycles64()) };
+	}
+
+	FDateTime EstimatedLocalDateTime() const
+	{
+		const FUdpMessagingTime RelNow = Now();
+		const FDateTime AbsNow = FDateTime::Now();
+		return AbsNow - (RelNow - *this);
 	}
 
 	FTimespan operator-(FUdpMessagingTime Other) const
 	{
-		return FTimespan::FromSeconds(RawValue - Other.RawValue);
+		return FTimespan::FromSeconds((Cycles - Other.Cycles) * FPlatformTime::GetSecondsPerCycle64());
 	}
 
 	auto operator+(FUdpMessagingTime Other) const = delete;
@@ -26,23 +34,25 @@ struct FUdpMessagingTime
 
 	FUdpMessagingTime operator-(FTimespan Other) const
 	{
-		return { .RawValue = RawValue - Other.GetTotalSeconds() };
+		const int64 Delta = FMath::RoundToInt64(Other.GetTotalSeconds() / FPlatformTime::GetSecondsPerCycle64());
+		return { .Cycles = Cycles - Delta };
 	}
 
 	FUdpMessagingTime operator+(FTimespan Other) const
 	{
-		return { .RawValue = RawValue + Other.GetTotalSeconds() };
+		const int64 Delta = FMath::RoundToInt64(Other.GetTotalSeconds() / FPlatformTime::GetSecondsPerCycle64());
+		return { .Cycles = Cycles + Delta };
 	}
 
 #if 0 // Enable some day when all our target toolchains support it.
-	/** Consistent compiler generated default comparison operations based on RawValue. */
+	/** Consistent compiler generated default comparison operations based on Cycles. */
 	friend auto operator<=>(FUdpMessagingTime, FUdpMessagingTime) = default;
 #else
-	bool operator==(FUdpMessagingTime Other) const { return RawValue == Other.RawValue; }
-	bool operator!=(FUdpMessagingTime Other) const { return RawValue != Other.RawValue; }
-	bool operator>(FUdpMessagingTime Other) const { return RawValue > Other.RawValue; }
-	bool operator>=(FUdpMessagingTime Other) const { return RawValue >= Other.RawValue; }
-	bool operator<(FUdpMessagingTime Other) const { return RawValue < Other.RawValue; }
-	bool operator<=(FUdpMessagingTime Other) const { return RawValue <= Other.RawValue; }
+	bool operator==(FUdpMessagingTime Other) const { return Cycles == Other.Cycles; }
+	bool operator!=(FUdpMessagingTime Other) const { return Cycles != Other.Cycles; }
+	bool operator>(FUdpMessagingTime Other) const { return Cycles > Other.Cycles; }
+	bool operator>=(FUdpMessagingTime Other) const { return Cycles >= Other.Cycles; }
+	bool operator<(FUdpMessagingTime Other) const { return Cycles < Other.Cycles; }
+	bool operator<=(FUdpMessagingTime Other) const { return Cycles <= Other.Cycles; }
 #endif
 };
diff --git a/Engine/Plugins/Messaging/UdpMessaging/Source/UdpMessaging/Private/Tests/UdpMessageSegmenterTest.cpp b/Engine/Plugins/Messaging/UdpMessaging/Source/UdpMessaging/Private/Tests/UdpMessageSegmenterTest.cpp
--- a/Engine/Plugins/Messaging/UdpMessaging/Source/UdpMessaging/Private/Tests/UdpMessageSegmenterTest.cpp
+++ b/Engine/Plugins/Messaging/UdpMessaging/Source/UdpMessaging/Private/Tests/UdpMessageSegmenterTest.cpp
@@ -39,10 +39,11 @@ void RunSegmentationTest(FAutomationTestBase& Test, uint32 MessageSize, uint16 S
 	// create and initialize segmenter
 	const uint16 SegmentSize = IntCastChecked<uint16>((uint32)SegmentPayloadSize + FUdpMessageSegmenter::GetDataSegmentHeaderSize());
 	FUdpMessageSegmenter Segmenter = FUdpMessageSegmenter(Message, SegmentSize);
-	Segmenter.Initialize();
 
-	Test.TestEqual(TEXT("Our calculation of segment count must match the segmenter"), Segmenter.GetSegmentCount(), NumSegments);
+	const bool bInitializeResult = Segmenter.TryInitialize();
+	Test.TestTrue(TEXT("Attempting to initialize the segmenter succeeds"), bInitializeResult);
 
+	Test.TestEqual(TEXT("Our calculation of segment count must match the segmenter"), Segmenter.GetSegmentCount(), NumSegments);
 
 	// invariants
 	{
diff --git a/Engine/Plugins/Messaging/UdpMessaging/Source/UdpMessaging/Private/Trace/UdpMessagingTrace.cpp b/Engine/Plugins/Messaging/UdpMessaging/Source/UdpMessaging/Private/Trace/UdpMessagingTrace.cpp
new file mode 100644
--- /dev/null
+++ b/Engine/Plugins/Messaging/UdpMessaging/Source/UdpMessaging/Private/Trace/UdpMessagingTrace.cpp
@@ -0,0 +1,150 @@
+// Copyright Epic Games, Inc. All Rights Reserved.
+
+#include "Trace/UdpMessagingTrace.h"
+
+
+namespace UE::UdpMessaging::Trace
+{
+	const TCHAR* LexToString(EUdpMessageLifecycleEvent InEvent)
+	{
+		switch (InEvent)
+		{
+		case EUdpMessageLifecycleEvent::EnqueueOutbound: return TEXT("EnqueueOutbound");
+		case EUdpMessageLifecycleEvent::BeginSerialize: return TEXT("BeginSerialize");
+		case EUdpMessageLifecycleEvent::EndSerialize: return TEXT("EndSerialize");
+		case EUdpMessageLifecycleEvent::ConsumeOutbound: return TEXT("ConsumeOutbound");
+		case EUdpMessageLifecycleEvent::Complete: return TEXT("Complete");
+
+		case EUdpMessageLifecycleEvent::FirstSegmentEnqueueInbound: return TEXT("FirstSegmentEnqueueInbound");
+		case EUdpMessageLifecycleEvent::BeginReassembly: return TEXT("BeginReassembly");
+		case EUdpMessageLifecycleEvent::EndReassembly: return TEXT("EndReassembly");
+		case EUdpMessageLifecycleEvent::PendingDelivery: return TEXT("PendingDelivery");
+		case EUdpMessageLifecycleEvent::BeginDeserialize: return TEXT("BeginDeserialize");
+		case EUdpMessageLifecycleEvent::Delivered: return TEXT("Delivered");
+
+		case EUdpMessageLifecycleEvent::Discarded: return TEXT("Discarded");
+
+		default:
+			return TEXT("Unknown");
+		}
+	}
+}
+
+
+#if UDPMESSAGINGTRACE_ENABLED
+
+#include "HAL/PlatformTime.h"
+#include "Interfaces/IPv4/IPv4Endpoint.h"
+#include "Misc/Guid.h"
+#include "Misc/TVariant.h"
+#include "ProfilingDebugging/StringsTrace.h"
+#include "Trace/Trace.inl"
+#include "UObject/TopLevelAssetPath.h"
+
+#include <type_traits>
+
+
+UE_TRACE_EVENT_BEGIN(UdpMessaging, DiscoveredNode, Important|NoSync)
+	UE_TRACE_EVENT_FIELD(uint64, Cycle)
+	UE_TRACE_EVENT_FIELD(uint16, ShortId) // Lightweight primary key, auto-increment, assumes we see < 64k unique nodes over life of app
+	UE_TRACE_EVENT_FIELD(uint16, DiscoveredByShortId) // Our local short ID which discovered this one; if 0, the "discovered" node is a new local processor
+	UE_TRACE_EVENT_FIELD(uint32, IdA)
+	UE_TRACE_EVENT_FIELD(uint32, IdB)
+	UE_TRACE_EVENT_FIELD(uint32, IdC)
+	UE_TRACE_EVENT_FIELD(uint32, IdD)
+	UE_TRACE_EVENT_FIELD(uint8,  IpAddrA)
+	UE_TRACE_EVENT_FIELD(uint8,  IpAddrB)
+	UE_TRACE_EVENT_FIELD(uint8,  IpAddrC)
+	UE_TRACE_EVENT_FIELD(uint8,  IpAddrD)
+	UE_TRACE_EVENT_FIELD(uint16, Port)
+	UE_TRACE_EVENT_FIELD(uint8,  ProtocolVersion)
+UE_TRACE_EVENT_END()
+
+
+UE_TRACE_EVENT_BEGIN(UdpMessaging, MessageSummary, NoSync)
+	UE_TRACE_EVENT_FIELD(uint16, SenderShortId)
+	UE_TRACE_EVENT_FIELD(uint16, RecipientShortId)
+	UE_TRACE_EVENT_FIELD(uint32, MessageId)
+	UE_TRACE_EVENT_FIELD(uint64, Size)
+	UE_TRACE_EVENT_FIELD(uint32, NumSegments)
+	UE_TRACE_EVENT_FIELD(uint32, Flags)
+UE_TRACE_EVENT_END()
+
+
+UE_TRACE_EVENT_BEGIN(UdpMessaging, MessageTypeInfo, NoSync)
+	UE_TRACE_EVENT_FIELD(uint16, SenderShortId)
+	UE_TRACE_EVENT_FIELD(uint16, RecipientShortId)
+	UE_TRACE_EVENT_FIELD(uint32, MessageId)
+	UE_TRACE_EVENT_REFERENCE_FIELD(Strings, FName, Package) // FTopLevelAssetPath.PackageName or None
+	UE_TRACE_EVENT_REFERENCE_FIELD(Strings, FName, Asset) // FTopLevelAssetPath.AssetName or combined
+UE_TRACE_EVENT_END()
+
+UE_TRACE_EVENT_BEGIN(UdpMessaging, MessageLifecycle)
+	UE_TRACE_EVENT_FIELD(uint64, Cycle)
+	UE_TRACE_EVENT_FIELD(uint16, SenderShortId)
+	UE_TRACE_EVENT_FIELD(uint16, RecipientShortId)
+	UE_TRACE_EVENT_FIELD(uint32, MessageId)
+	UE_TRACE_EVENT_FIELD(uint8,  LifecycleEvent)
+UE_TRACE_EVENT_END()
+
+
+namespace UE::UdpMessaging::Trace
+{
+	void OutputDiscoveredNode(uint16 ShortId, uint16 DiscoveredByShortId, const FGuid& NodeId, const FIPv4Endpoint& Endpoint, uint8 ProtocolVersion)
+	{
+		UE_TRACE_LOG(UdpMessaging, DiscoveredNode, MessagingChannel)
+			<< DiscoveredNode.Cycle(FPlatformTime::Cycles64())
+			<< DiscoveredNode.ShortId(ShortId)
+			<< DiscoveredNode.DiscoveredByShortId(DiscoveredByShortId)
+			<< DiscoveredNode.IdA(NodeId.A)
+			<< DiscoveredNode.IdB(NodeId.B)
+			<< DiscoveredNode.IdC(NodeId.C)
+			<< DiscoveredNode.IdD(NodeId.D)
+			<< DiscoveredNode.IpAddrA(Endpoint.Address.A)
+			<< DiscoveredNode.IpAddrB(Endpoint.Address.B)
+			<< DiscoveredNode.IpAddrC(Endpoint.Address.C)
+			<< DiscoveredNode.IpAddrD(Endpoint.Address.D)
+			<< DiscoveredNode.Port(Endpoint.Port)
+			<< DiscoveredNode.ProtocolVersion(ProtocolVersion);
+	}
+
+	void OutputMessageSummary(uint16 SenderShortId, uint16 RecipientShortId, uint32 MessageId, uint64 Size, uint32 NumSegments, EMessageFlags Flags)
+	{
+		UE_TRACE_LOG(UdpMessaging, MessageSummary, MessagingChannel)
+			<< MessageSummary.SenderShortId(SenderShortId)
+			<< MessageSummary.RecipientShortId(RecipientShortId)
+			<< MessageSummary.MessageId(MessageId)
+			<< MessageSummary.Size(Size)
+			<< MessageSummary.NumSegments(NumSegments)
+			<< MessageSummary.Flags(static_cast<std::underlying_type_t<EMessageFlags>>(Flags));
+	}
+
+	void OutputMessageTypeInfo(uint16 SenderShortId, uint16 RecipientShortId, uint32 MessageId, const TVariant<FName, FTopLevelAssetPath>& TypeInfo)
+	{
+		const FTopLevelAssetPath* Path = TypeInfo.TryGet<FTopLevelAssetPath>();
+		const FName Package = Path ? Path->GetPackageName() : NAME_None;
+		const FName Asset = Path ? Path->GetAssetName() : TypeInfo.Get<FName>();
+
+		const UE::Trace::FEventRef32 PackageRef = FStringTrace::GetNameRef(Package);
+		const UE::Trace::FEventRef32 AssetRef = FStringTrace::GetNameRef(Asset);
+
+		UE_TRACE_LOG(UdpMessaging, MessageTypeInfo, MessagingChannel)
+			<< MessageTypeInfo.SenderShortId(SenderShortId)
+			<< MessageTypeInfo.RecipientShortId(RecipientShortId)
+			<< MessageTypeInfo.MessageId(MessageId)
+			<< MessageTypeInfo.Package(PackageRef)
+			<< MessageTypeInfo.Asset(AssetRef);
+	}
+
+	void OutputMessageLifecycle(uint64 Cycle, uint16 SenderShortId, uint16 RecipientShortId, uint32 MessageId, EUdpMessageLifecycleEvent Event)
+	{
+		UE_TRACE_LOG(UdpMessaging, MessageLifecycle, MessagingChannel)
+			<< MessageLifecycle.Cycle(Cycle)
+			<< MessageLifecycle.SenderShortId(SenderShortId)
+			<< MessageLifecycle.RecipientShortId(RecipientShortId)
+			<< MessageLifecycle.MessageId(MessageId)
+			<< MessageLifecycle.LifecycleEvent(static_cast<std::underlying_type_t<EUdpMessageLifecycleEvent>>(Event));
+	}
+};
+
+#endif // #if UDPMESSAGINGTRACE_ENABLED
diff --git a/Engine/Plugins/Messaging/UdpMessaging/Source/UdpMessaging/Private/Transport/UdpDeserializedMessage.cpp b/Engine/Plugins/Messaging/UdpMessaging/Source/UdpMessaging/Private/Transport/UdpDeserializedMessage.cpp
--- a/Engine/Plugins/Messaging/UdpMessaging/Source/UdpMessaging/Private/Transport/UdpDeserializedMessage.cpp
+++ b/Engine/Plugins/Messaging/UdpMessaging/Source/UdpMessaging/Private/Transport/UdpDeserializedMessage.cpp
@@ -12,7 +12,7 @@
 #include "Shared/UdpMessagingSettings.h"
 #include "UObject/TopLevelAssetPath.h"
 #include "UObject/WeakObjectPtrTemplates.h"
-#include "UdpMessagingTracing.h"
+#include "Trace/UdpMessagingTrace.h"
 
 /* FUdpDeserializedMessage structors
 *****************************************************************************/
diff --git a/Engine/Plugins/Messaging/UdpMessaging/Source/UdpMessaging/Private/Transport/UdpMessageProcessor.cpp b/Engine/Plugins/Messaging/UdpMessaging/Source/UdpMessaging/Private/Transport/UdpMessageProcessor.cpp
--- a/Engine/Plugins/Messaging/UdpMessaging/Source/UdpMessaging/Private/Transport/UdpMessageProcessor.cpp
+++ b/Engine/Plugins/Messaging/UdpMessaging/Source/UdpMessaging/Private/Transport/UdpMessageProcessor.cpp
@@ -13,6 +13,7 @@
 #include "UdpMessagingPrivate.h"
 #include "Shared/SocketSender.h"
 #include "Shared/UdpMessagingSettings.h"
+#include "Trace/UdpMessagingTrace.h"
 #include "Transport/UdpMessageBeacon.h"
 #include "Transport/UdpMessageSegmenter.h"
 #include "Transport/UdpReassembledMessage.h"
@@ -25,6 +26,10 @@
 /* FUdpMessageHelloSender static initialization
  *****************************************************************************/
 
+#if UDPMESSAGINGTRACE_ENABLED
+std::atomic<uint16> FUdpMessageProcessor::NextNodeTraceId(1);
+#endif
+
 const int32 FUdpMessageProcessor::DeadHelloIntervals = 5;
 
 TAutoConsoleVariable<int32> CVarFakeSocketError(
@@ -252,7 +257,7 @@ bool FUdpMessageProcessor::EnqueueInboundSegment(const TSharedPtr<FArrayReader,
 		return false;
 	}
 
-	InboundSegments.Enqueue(FInboundSegment(Data, InSender));
+	InboundSegments.Enqueue(Data, InSender, FUdpMessagingTime::Now());
 
 	WorkEvent->Trigger();
 
@@ -269,6 +274,8 @@ bool FUdpMessageProcessor::EnqueueOutboundMessage(const TSharedRef<IMessageConte
 	TMap<uint8, TArray<FGuid>> RecipientPerVersions = GetRecipientsPerProtocolVersion(Recipients);
 	for (const auto& RecipientVersion : RecipientPerVersions)
 	{
+		const FUdpMessagingTime EnqueueTime = FUdpMessagingTime::Now();
+
 		// Create a message to serialize using that protocol version
 		TSharedRef<FUdpSerializedMessage, ESPMode::ThreadSafe> SerializedMessage = MakeShared<FUdpSerializedMessage, ESPMode::ThreadSafe>(MessageFormat, RecipientVersion.Key, MessageContext->GetFlags());
 
@@ -276,7 +283,7 @@ bool FUdpMessageProcessor::EnqueueOutboundMessage(const TSharedRef<IMessageConte
 		TGraphTask<FUdpSerializeMessageTask>::CreateTask().ConstructAndDispatchWhenReady(MessageContext, SerializedMessage, WorkEvent);
 
 		// Enqueue the message
-		OutboundMessages.Enqueue(FOutboundMessage(SerializedMessage, RecipientVersion.Value, MessageContext->GetFlags()));
+		OutboundMessages.Enqueue(SerializedMessage, RecipientVersion.Value, MessageContext->GetFlags(), EnqueueTime);
 	}
 
 	return true;
@@ -356,6 +363,14 @@ bool FUdpMessageProcessor::Init()
 
 		Beacon = MakeUnique<FUdpMessageBeacon>(SocketSender.Get(), LocalNodeId, MulticastEndpoint);
 
+#if UDPMESSAGINGTRACE_ENABLED
+		LocalNodeTraceId = NextNodeTraceId++;
+#endif
+
+		TRACE_UDPMESSAGING_DISCOVERED_NODE(
+			LocalNodeTraceId, 0, LocalNodeId, FIPv4Endpoint(SocketSender->GetAddress()), UDP_MESSAGING_TRANSPORT_PROTOCOL_VERSION
+		);
+
 		// Current protocol version 18
 		SupportedProtocolVersions.Add(UDP_MESSAGING_TRANSPORT_PROTOCOL_VERSION);
 		// Support Protocol version 10, 11, 12, 13, 14, 15, 16, 17
@@ -512,9 +527,10 @@ void FUdpMessageProcessor::ConsumeInboundSegments()
 	const FTimespan MaxWorkTimespan = FTimespan::FromMilliseconds(5);
 	const FUdpMessagingTime EndWorkTime = FUdpMessagingTime::Now() + MaxWorkTimespan;
 
-	FInboundSegment Segment;
-	while (InboundSegments.Dequeue(Segment))
+	while (TOptional<FInboundSegment> SegmentOptional = InboundSegments.Dequeue())
 	{
+		FInboundSegment& Segment = SegmentOptional.GetValue();
+
 		// quick hack for TTP# 247103
 		if (!Segment.Data.IsValid())
 		{
@@ -536,9 +552,15 @@ void FUdpMessageProcessor::ConsumeInboundSegments()
 
 			if (!NodeInfo.NodeId.IsValid())
 			{
+#if UDPMESSAGINGTRACE_ENABLED
+				NodeInfo.NodeTraceId = NextNodeTraceId++;
+#endif
 				NodeInfo.NodeId = Header.SenderNodeId;
 				NodeInfo.ProtocolVersion = Header.ProtocolVersion;
 
+				TRACE_UDPMESSAGING_DISCOVERED_NODE(NodeInfo.NodeTraceId, LocalNodeTraceId,
+					NodeInfo.NodeId, Segment.Sender, NodeInfo.ProtocolVersion);
+
 				const uint16 MaxSegmentSizeV10_V17 = 2048 - UE::UdpMessaging::PacketHeaderBytes;
 				const uint16 MaxSegmentSize = GetDefault<UUdpMessagingSettings>()->GetMaxPacketSize()
 					- UE::UdpMessaging::PacketHeaderBytes;
@@ -619,6 +641,17 @@ void FUdpMessageProcessor::ConsumeInboundSegments()
 	}
 }
 
+int32 FUdpMessageProcessor::AssignNextMessageId()
+{
+	++LastSentMessage;
+	if (LastSentMessage < 0)
+	{
+		// Prevent negative message ids
+		LastSentMessage = 0;
+	}
+	return LastSentMessage;
+}
+
 void FUdpMessageProcessor::ConsumeOneOutboundMessage(const FOutboundMessage& OutboundMessage)
 {
 	const auto FindNodeWithLog = [this](const FGuid &Id)
@@ -638,38 +671,56 @@ void FUdpMessageProcessor::ConsumeOneOutboundMessage(const FOutboundMessage& Out
 
 	const bool bIsReliable = EnumHasAnyFlags(OutboundMessage.MessageFlags, EMessageFlags::Reliable);
 
-	++LastSentMessage;
-	if (LastSentMessage < 0)
-	{
-		// Prevent negative message ids
-		LastSentMessage = 0;
-	}
+	const int32 MessageId = AssignNextMessageId();
 
 	for (FNodeInfo* RecipientNodeInfo : Recipients)
 	{
+		TRACE_UDPMESSAGING_MESSAGE_LIFECYCLE(
+			OutboundMessage.EnqueueTime.Cycles,
+			LocalNodeTraceId,
+			RecipientNodeInfo->NodeTraceId,
+			MessageId,
+			UE::UdpMessaging::Trace::EUdpMessageLifecycleEvent::EnqueueOutbound
+		);
+
+		TRACE_UDPMESSAGING_MESSAGE_LIFECYCLE(
+			FPlatformTime::Cycles64(),
+			LocalNodeTraceId,
+			RecipientNodeInfo->NodeTraceId,
+			MessageId,
+			UE::UdpMessaging::Trace::EUdpMessageLifecycleEvent::ConsumeOutbound
+		);
+
 		UE_LOG(LogUdpMessaging, VeryVerbose, TEXT("Passing %" INT64_FMT " byte message to be segment-sent %d to %s with id %s"),
 				OutboundMessage.SerializedMessage->TotalSize(),
-				LastSentMessage,
+				MessageId,
 				*RecipientNodeInfo->Endpoint.ToString(),
 				*RecipientNodeInfo->NodeId.ToString());
 
 		if (!bIsReliable && !RecipientNodeInfo->CanSendSegments())
 		{
 			// Discard unreliable messages that cannot be sent.
+			TRACE_UDPMESSAGING_MESSAGE_LIFECYCLE(
+				FPlatformTime::Cycles64(),
+				LocalNodeTraceId,
+				RecipientNodeInfo->NodeTraceId,
+				MessageId,
+				UE::UdpMessaging::Trace::EUdpMessageLifecycleEvent::Discarded
+			);
 			continue;
 		}
 
 		TSharedPtr<FUdpMessageSegmenter> Segmenter = MakeShared<FUdpMessageSegmenter>(
 			OutboundMessage.SerializedMessage.ToSharedRef(), RecipientNodeInfo->MaxSegmentSize);
 
-		RecipientNodeInfo->Segmenters.Add(LastSentMessage, Segmenter);
+		RecipientNodeInfo->Segmenters.Add(MessageId, Segmenter);
 		if (bIsReliable)
 		{
-			RecipientNodeInfo->ReliableWorkQueue.Add(LastSentMessage);
+			RecipientNodeInfo->ReliableWorkQueue.Add(MessageId);
 		}
 		else
 		{
-			RecipientNodeInfo->UnreliableWorkQueue.Add(LastSentMessage);
+			RecipientNodeInfo->UnreliableWorkQueue.Add(MessageId);
 		}
 	}
 }
@@ -678,9 +729,9 @@ void FUdpMessageProcessor::ConsumeOutboundMessages()
 {
 	SCOPED_MESSAGING_TRACE(FUdpMessageProcessor_ConsumeOutputMessages);
 	
-	FOutboundMessage OutboundMessage;
-	while (OutboundMessages.Dequeue(OutboundMessage))
+	while (TOptional<FOutboundMessage> OutboundMessageOptional = OutboundMessages.Dequeue())
 	{
+		const FOutboundMessage& OutboundMessage = OutboundMessageOptional.GetValue();
 		ConsumeOneOutboundMessage(OutboundMessage);
 	}
 }
@@ -715,7 +766,18 @@ void FUdpMessageProcessor::ProcessAbortSegment(FInboundSegment& Segment, FNodeIn
 	FUdpMessageSegment::FAbortChunk AbortChunk;
 	AbortChunk.Serialize(*Segment.Data, NodeInfo.ProtocolVersion);
 
-	NodeInfo.Segmenters.Remove(AbortChunk.MessageId);
+	const int32 NumRemoved = NodeInfo.Segmenters.Remove(AbortChunk.MessageId);
+
+	if (NumRemoved)
+	{
+		TRACE_UDPMESSAGING_MESSAGE_LIFECYCLE(
+			FPlatformTime::Cycles64(),
+			LocalNodeTraceId,
+			NodeInfo.NodeTraceId,
+			AbortChunk.MessageId,
+			UE::UdpMessaging::Trace::EUdpMessageLifecycleEvent::Discarded
+		);
+	}
 }
 
 
@@ -737,7 +799,17 @@ void FUdpMessageProcessor::ProcessAcknowledgeSegment(FInboundSegment& Segment, F
 		SendSegmenterStatsToListeners(AcknowledgeChunk.MessageId, NodeInfo.NodeId, Segmenter);
 	}
 
-	NodeInfo.Segmenters.Remove(AcknowledgeChunk.MessageId);
+	const int32 NumRemoved = NodeInfo.Segmenters.Remove(AcknowledgeChunk.MessageId);
+	if (NumRemoved)
+	{
+		TRACE_UDPMESSAGING_MESSAGE_LIFECYCLE(
+			FPlatformTime::Cycles64(),
+			LocalNodeTraceId,
+			NodeInfo.NodeTraceId,
+			AcknowledgeChunk.MessageId,
+			UE::UdpMessaging::Trace::EUdpMessageLifecycleEvent::Complete
+		);
+	}
 
 	UE_LOG(LogUdpMessaging, Verbose, TEXT("Received Acknowledge for %d from %s"), AcknowledgeChunk.MessageId , *NodeInfo.NodeId.ToString());
 }
@@ -751,7 +823,17 @@ void FUdpMessageProcessor::ProcessAcknowledgeSegmentsSegment(FInboundSegment& Se
 	AcknowledgeChunk.Serialize(*Segment.Data, NodeInfo.ProtocolVersion);
 
 	UE_LOG(LogUdpMessaging, Verbose, TEXT("Received AcknowledgeSegments for %d from %s"), AcknowledgeChunk.MessageId, *NodeInfo.NodeId.ToString());
-	NodeInfo.MarkAcks(AcknowledgeChunk.MessageId, AcknowledgeChunk.Segments, FUdpMessagingTime::Now());
+	const bool bIsComplete = NodeInfo.MarkAcks(AcknowledgeChunk.MessageId, AcknowledgeChunk.Segments, FUdpMessagingTime::Now());
+	if (bIsComplete)
+	{
+		TRACE_UDPMESSAGING_MESSAGE_LIFECYCLE(
+			FPlatformTime::Cycles64(),
+			LocalNodeTraceId,
+			NodeInfo.NodeTraceId,
+			AcknowledgeChunk.MessageId,
+			UE::UdpMessaging::Trace::EUdpMessageLifecycleEvent::Complete
+		);
+	}
 }
 
 
@@ -790,12 +872,19 @@ void FUdpMessageProcessor::ProcessDataSegment(FInboundSegment& Segment, FNodeInf
 		return;
 	}
 
-	TSharedPtr<FUdpReassembledMessage, ESPMode::ThreadSafe>& ReassembledMessage = NodeInfo.ReassembledMessages.FindOrAdd(DataChunk.MessageId);
+	TSharedPtr<FUdpReassembledMessage>& ReassembledMessage = NodeInfo.ReassembledMessages.FindOrAdd(DataChunk.MessageId);
 
 	// Reassemble message
 	if (!ReassembledMessage.IsValid())
 	{
-		ReassembledMessage = MakeShared<FUdpReassembledMessage, ESPMode::ThreadSafe>(NodeInfo.ProtocolVersion, DataChunk.MessageFlags, DataChunk.MessageSize, DataChunk.TotalSegments, DataChunk.Sequence, Segment.Sender);
+		ReassembledMessage = MakeShared<FUdpReassembledMessage>(
+			NodeInfo.ProtocolVersion,
+			DataChunk.MessageFlags,
+			DataChunk.MessageSize,
+			DataChunk.TotalSegments,
+			DataChunk.Sequence,
+			Segment.Sender
+		);
 
 		if (ReassembledMessage->IsMalformed())
 		{
@@ -806,6 +895,14 @@ void FUdpMessageProcessor::ProcessDataSegment(FInboundSegment& Segment, FNodeInf
 			ReassembledMessage.Reset();
 			return;
 		}
+
+#if UDPMESSAGINGTRACE_ENABLED
+		ReassembledMessage->GetTraceMetadata() = { .SenderShortId = NodeInfo.NodeTraceId, .RecipientShortId = LocalNodeTraceId, .MessageId = DataChunk.MessageId };
+#endif
+
+		TRACE_UDPMESSAGING_MESSAGE_LIFECYCLE(Segment.EnqueueTime.Cycles, NodeInfo.NodeTraceId, LocalNodeTraceId, DataChunk.MessageId, UE::UdpMessaging::Trace::EUdpMessageLifecycleEvent::FirstSegmentEnqueueInbound);
+		TRACE_UDPMESSAGING_MESSAGE_LIFECYCLE(FPlatformTime::Cycles64(), NodeInfo.NodeTraceId, LocalNodeTraceId, DataChunk.MessageId, UE::UdpMessaging::Trace::EUdpMessageLifecycleEvent::BeginReassembly);
+		TRACE_UDPMESSAGING_MESSAGE_SUMMARY(NodeInfo.NodeTraceId, LocalNodeTraceId, DataChunk.MessageId, DataChunk.MessageSize, DataChunk.TotalSegments, DataChunk.MessageFlags);
 	}
 
 	NodeInfo.Statistics.TotalBytesReceived += DataChunk.Data.Num();
@@ -842,6 +939,8 @@ void FUdpMessageProcessor::ProcessDataSegment(FInboundSegment& Segment, FNodeInf
 		return;
 	}
 
+	TRACE_UDPMESSAGING_MESSAGE_LIFECYCLE(FPlatformTime::Cycles64(), NodeInfo.NodeTraceId, LocalNodeTraceId, DataChunk.MessageId, UE::UdpMessaging::Trace::EUdpMessageLifecycleEvent::EndReassembly);
+
 	UE_LOG(LogUdpMessaging, Verbose, TEXT("FUdpMessageProcessor::ProcessDataSegment: Reassembled %d bytes message %s with %d for %s (%s)"),
 		ReassembledMessage->GetData().Num(),
 		*ReassembledMessage->Describe(),
@@ -984,8 +1083,15 @@ bool FUdpMessageProcessor::LookupAndCacheMessageType(TSharedPtr<FUdpReassembledM
 	}
 
 	FNameOrFTopLevel AssetPath = FUdpDeserializedMessage::PeekMessageTypeInfoName(*ReassembledMessage);
-	FString PathAsString = Visit([](auto&& Path) -> FString { return Path.ToString(); }, AssetPath);
 
+	TRACE_UDPMESSAGING_MESSAGE_TYPEINFO(
+		ReassembledMessage->GetTraceMetadata().SenderShortId,
+		ReassembledMessage->GetTraceMetadata().RecipientShortId,
+		ReassembledMessage->GetTraceMetadata().MessageId,
+		AssetPath
+	);
+
+	FString PathAsString = Visit([](auto&& Path) -> FString { return Path.ToString(); }, AssetPath);
 	if (TWeakObjectPtr<UScriptStruct>* TypeInfo = CachedTypeInfoMap.Find(PathAsString))
 	{
 		ReassembledMessage->SetMessageTypeInfo(*TypeInfo);
@@ -1004,7 +1110,7 @@ bool FUdpMessageProcessor::LookupAndCacheMessageType(TSharedPtr<FUdpReassembledM
 	return false;
 }
 
-void FUdpMessageProcessor::TryDeliverMessage(const TSharedPtr<FUdpReassembledMessage, ESPMode::ThreadSafe>& ReassembledMessage, FNodeInfo& NodeInfo)
+void FUdpMessageProcessor::TryDeliverMessage(const TSharedPtr<FUdpReassembledMessage>& ReassembledMessage, FNodeInfo& NodeInfo)
 {
 	// Do not deliver message while saving or garbage collecting since those deliveries will fail anyway...
 	if (ReassembledMessage->GetMessageTypeInfo() == nullptr && (UE::IsSavingPackage() || IsGarbageCollecting()))
@@ -1013,10 +1119,30 @@ void FUdpMessageProcessor::TryDeliverMessage(const TSharedPtr<FUdpReassembledMes
 		return;
 	}
 
+	TRACE_UDPMESSAGING_MESSAGE_LIFECYCLE(
+		FPlatformTime::Cycles64(),
+		NodeInfo.NodeTraceId,
+		LocalNodeTraceId,
+		ReassembledMessage->GetTraceMetadata().MessageId,
+		UE::UdpMessaging::Trace::EUdpMessageLifecycleEvent::PendingDelivery
+	);
+
 	if (NodeInfo.NodeId.IsValid())
 	{
 		MessageReassembledDelegate.ExecuteIfBound(ReassembledMessage, nullptr, NodeInfo.NodeId);
 	}
+	else
+	{
+		UE_LOG(LogUdpMessaging, Display, TEXT("Unable to deliver %s due to invalid NodeId"), *ReassembledMessage->Describe());
+
+		TRACE_UDPMESSAGING_MESSAGE_LIFECYCLE(
+			FPlatformTime::Cycles64(),
+			ReassembledMessage->GetTraceMetadata().SenderShortId,
+			ReassembledMessage->GetTraceMetadata().RecipientShortId,
+			ReassembledMessage->GetTraceMetadata().MessageId,
+			UE::UdpMessaging::Trace::EUdpMessageLifecycleEvent::Discarded
+		);
+	}
 
 	// Mark the message pending delivery but do not remove it from the list yet, this is to prevent the double delivery of reliable message.
 	ReassembledMessage->MarkPendingDelivery();
@@ -1251,7 +1377,7 @@ void FUdpMessageProcessor::UpdateKnownNodes()
 	}
 }
 
-TSharedPtr<FUdpMessageSegmenter>* FUdpMessageProcessor::GetInitializedSegmenter(FNodeInfo& NodeInfo, int32 MessageId)
+TSharedPtr<FUdpMessageSegmenter>* FUdpMessageProcessor::GetSegmenter(FNodeInfo& NodeInfo, int32 MessageId)
 {
 	TSharedPtr<FUdpMessageSegmenter>* FoundSegmenter = NodeInfo.Segmenters.Find(MessageId);
 	if (!FoundSegmenter)
@@ -1263,13 +1389,26 @@ TSharedPtr<FUdpMessageSegmenter>* FUdpMessageProcessor::GetInitializedSegmenter(
 	TSharedPtr<FUdpMessageSegmenter>& Segmenter = *FoundSegmenter;
 	if (!Segmenter->IsInitialized())
 	{
-		Segmenter->Initialize();
+		// This is essentially polling on the async serialization to complete.
+		if (Segmenter->TryInitialize())
+		{
+#if UDPMESSAGINGTRACE_ENABLED
+			Segmenter->EmitPostInitializeTraceEvents(LocalNodeTraceId, NodeInfo.NodeTraceId, MessageId);
+#endif
+		}
 	}
 
 	if (Segmenter->IsInvalid())
 	{
 		UE_LOG(LogUdpMessaging, Warning, TEXT("FUdpMessageProcessor::GetInitializedSegments segmenter is invalid. Removing from send queue."));
 		NodeInfo.Segmenters.Remove(MessageId);
+		TRACE_UDPMESSAGING_MESSAGE_LIFECYCLE(
+			FPlatformTime::Cycles64(),
+			LocalNodeTraceId,
+			NodeInfo.NodeTraceId,
+			MessageId,
+			UE::UdpMessaging::Trace::EUdpMessageLifecycleEvent::Discarded
+		);
 		return nullptr;
 	}
 
@@ -1331,7 +1470,7 @@ FSentSegmentInfo FUdpMessageProcessor::SendNextSegmentsForMessageId(
 	int32 MaxSegments /* = 1 */
 )
 {
-	TSharedPtr<FUdpMessageSegmenter>* FoundSegmenter = GetInitializedSegmenter(NodeInfo, MessageId);
+	TSharedPtr<FUdpMessageSegmenter>* FoundSegmenter = GetSegmenter(NodeInfo, MessageId);
 	if (!FoundSegmenter)
 	{
 		// It's possible that a segmenter was fully ack and removed from the segmenter list.
@@ -1349,6 +1488,13 @@ FSentSegmentInfo FUdpMessageProcessor::SendNextSegmentsForMessageId(
 			UE_LOG(LogUdpMessaging, Error,
 				TEXT("FUdpMessageProcessor::SendNextSegmentsForMessageId: Node %s initially asked for protocol version %d but changed to %d. Cannot deliver message %d"),
 				*NodeInfo.Endpoint.ToString(), Segmenter->GetProtocolVersion(), Header.ProtocolVersion, MessageId);
+			TRACE_UDPMESSAGING_MESSAGE_LIFECYCLE(
+				FPlatformTime::Cycles64(),
+				LocalNodeTraceId,
+				NodeInfo.NodeTraceId,
+				MessageId,
+				UE::UdpMessaging::Trace::EUdpMessageLifecycleEvent::Discarded
+			);
 			NodeInfo.Segmenters.Remove(MessageId);
 			return FSentSegmentInfo(MessageId);
 		}
@@ -1424,6 +1570,13 @@ FSentSegmentInfo FUdpMessageProcessor::SendNextSegmentsForMessageId(
 			if (Segmenter->IsSendingComplete() && !bIsReliable)
 			{
 				// Not reliably sent so we don't need to wait for acks.
+				TRACE_UDPMESSAGING_MESSAGE_LIFECYCLE(
+					FPlatformTime::Cycles64(),
+					LocalNodeTraceId,
+					NodeInfo.NodeTraceId,
+					MessageId,
+					UE::UdpMessaging::Trace::EUdpMessageLifecycleEvent::Complete
+				);
 				UE_LOG(LogUdpMessaging, VeryVerbose, TEXT("FUdpMessageProcessor::UpdateSegmenters: Finished with message segmenter for %s"), *NodeInfo.NodeId.ToString());
 				NodeInfo.Segmenters.Remove(MessageId);
 			}
@@ -1698,6 +1851,13 @@ bool FUdpMessageProcessor::UpdateReassemblers(FNodeInfo& NodeInfo)
 
 			if (!ReassembledMessage->IsPendingDelivery())
 			{
+				TRACE_UDPMESSAGING_MESSAGE_LIFECYCLE(
+					FPlatformTime::Cycles64(),
+					ReassembledMessagePtr->GetTraceMetadata().SenderShortId,
+					ReassembledMessagePtr->GetTraceMetadata().RecipientShortId,
+					ReassembledMessagePtr->GetTraceMetadata().MessageId,
+					UE::UdpMessaging::Trace::EUdpMessageLifecycleEvent::Discarded
+				);
 				PendingTypeLookup.Remove(ReassembledMessagePtr);
 				It.RemoveCurrent();
 			}
diff --git a/Engine/Plugins/Messaging/UdpMessaging/Source/UdpMessaging/Private/Transport/UdpMessageProcessor.h b/Engine/Plugins/Messaging/UdpMessaging/Source/UdpMessaging/Private/Transport/UdpMessageProcessor.h
--- a/Engine/Plugins/Messaging/UdpMessaging/Source/UdpMessaging/Private/Transport/UdpMessageProcessor.h
+++ b/Engine/Plugins/Messaging/UdpMessaging/Source/UdpMessaging/Private/Transport/UdpMessageProcessor.h
@@ -22,7 +22,9 @@
 #include "Shared/UdpMessagingTime.h"
 #include "UdpMessageSegmenter.h"
 #include "UdpMessagingPrivate.h"
-#include "UdpMessagingTracing.h"
+#include "Trace/UdpMessagingTrace.h"
+
+#include <atomic>
 
 
 class IMessageContext;
@@ -160,7 +162,7 @@ class FUdpMessageProcessor
 		uint16 MaxSegmentSize;
 
 		/** Holds the collection of reassembled messages. */
-		TMap<int32, TSharedPtr<FUdpReassembledMessage, ESPMode::ThreadSafe>> ReassembledMessages;
+		TMap<int32, TSharedPtr<FUdpReassembledMessage>> ReassembledMessages;
 
 		/** Holds the collection of message segmenters. */
 		TMap<int32, TSharedPtr<FUdpMessageSegmenter>> Segmenters;
@@ -217,6 +219,11 @@ class FUdpMessageProcessor
 		/** Various transport statistics for this endpoint */
 		FMessageTransportStatistics Statistics;
 
+#if UDPMESSAGINGTRACE_ENABLED
+		/** ID used for trace events. */
+		uint16 NodeTraceId;
+#endif
+
 		/** Default constructor. */
 		FNodeInfo();
 
@@ -386,8 +393,9 @@ class FUdpMessageProcessor
 		/**
 		 * For a given list of segments to acknowledge record how long it took to receive a response and update our
 		 * average round trip time. We also use this opportunity to calculate any loss and update our window size.
+		 * @return true if the message has been fully acknowledged, false otherwise.
 		 */
-		void MarkAcks(int32 MessageId, const TArray<uint32>& Segments, const FUdpMessagingTime InCurrentTime)
+		bool MarkAcks(int32 MessageId, const TArray<uint32>& Segments, const FUdpMessagingTime InCurrentTime)
 		{
 			TSharedPtr<FUdpMessageSegmenter> Segmenter = MarkAcksOnSegmenter(MessageId, Segments, InCurrentTime);
 			FTimespan MaxSpan(0);
@@ -415,6 +423,15 @@ class FUdpMessageProcessor
 
 			const uint32 LostSegmentCount = ExpireSegments(InCurrentTime);
 			ComputeWindowSize(Acks, LostSegmentCount);
+
+			if (Segmenter)
+			{
+				return Segmenter->IsSendingComplete() && Segmenter->AreAcknowledgementsComplete();
+			}
+			else
+			{
+				return false;
+			}
 		}
 
 		/**
@@ -530,18 +547,26 @@ class FUdpMessageProcessor
 	struct FInboundSegment
 	{
 		/** Holds the segment data. */
-		TSharedPtr<FArrayReader, ESPMode::ThreadSafe> Data;
+		TSharedPtr<FArrayReader> Data;
 
 		/** Holds the sender's network endpoint. */
 		FIPv4Endpoint Sender;
 
-		/** Default constructor. */
-		FInboundSegment() { }
+#if UDPMESSAGINGTRACE_ENABLED
+		const FUdpMessagingTime EnqueueTime;
+#endif
 
 		/** Creates and initializes a new instance. */
-		FInboundSegment(const TSharedPtr<FArrayReader, ESPMode::ThreadSafe>& InData, const FIPv4Endpoint& InSender)
+		FInboundSegment(
+			const TSharedPtr<FArrayReader>& InData,
+			const FIPv4Endpoint& InSender,
+			FUdpMessagingTime InEnqueueTime
+		)
 			: Data(InData)
 			, Sender(InSender)
+#if UDPMESSAGINGTRACE_ENABLED
+			, EnqueueTime(InEnqueueTime)
+#endif
 		{ }
 	};
 
@@ -550,28 +575,37 @@ class FUdpMessageProcessor
 	struct FOutboundMessage
 	{
 		/** Holds the serialized message. */
-		TSharedPtr<FUdpSerializedMessage, ESPMode::ThreadSafe> SerializedMessage;
+		TSharedPtr<FUdpSerializedMessage> SerializedMessage;
 
 		/** Holds the recipients. */
 		TArray<FGuid> RecipientIds;
 
 		EMessageFlags MessageFlags;
 
-		/** Default constructor. */
-		FOutboundMessage() { }
+#if UDPMESSAGINGTRACE_ENABLED
+		const FUdpMessagingTime EnqueueTime;
+#endif
 
 		/** Creates and initializes a new instance. */
-		FOutboundMessage(TSharedPtr<FUdpSerializedMessage, ESPMode::ThreadSafe> InSerializedMessage, const TArray<FGuid>& InRecipientIds, EMessageFlags InFlags)
+		FOutboundMessage(
+			const TSharedRef<FUdpSerializedMessage>& InSerializedMessage,
+			const TArray<FGuid>& InRecipientIds,
+			EMessageFlags InFlags,
+			FUdpMessagingTime InEnqueueTime
+		)
 			: SerializedMessage(InSerializedMessage)
 			, RecipientIds(InRecipientIds)
 			, MessageFlags(InFlags)
+#if UDPMESSAGINGTRACE_ENABLED
+			, EnqueueTime(InEnqueueTime)
+#endif
 		{ }
 	};
 
 	bool FillDataChunk(FUdpMessageSegment::FDataChunk& OutDataChunk, FNodeInfo& NodeInfo, TSharedPtr<FUdpMessageSegmenter>& Segmenter, int32 MessageId);
 
-	/** Returns an initialized segment ready for sending or ready to receive ack. */
-	TSharedPtr<FUdpMessageSegmenter>* GetInitializedSegmenter(FNodeInfo& NodeInfo, int32 MessageId);
+	/** Returns the specified segmenter if valid, and attempts to initialize it if necessary. */
+	TSharedPtr<FUdpMessageSegmenter>* GetSegmenter(FNodeInfo& NodeInfo, int32 MessageId);
 
 	int32 ClampMaxSegments(int32 MaxSegments, const TSharedPtr<FUdpMessageSegmenter>& Segmenter, const FNodeInfo& NodeInfo) const;
 
@@ -905,12 +939,12 @@ protected:
 	 * @param ReassembledMessage The message to deliver.
 	 * @param NodeInfo Details of the node that sent the message.
 	 */
-	void TryDeliverMessage(const TSharedPtr<FUdpReassembledMessage, ESPMode::ThreadSafe>& ReassembledMessage, FNodeInfo& NodeInfo);
+	void TryDeliverMessage(const TSharedPtr<FUdpReassembledMessage>& ReassembledMessage, FNodeInfo& NodeInfo);
 
 	/**
 	 * Lookup the reassembled message type and try to determine the source struct type.
 	 */
-	bool LookupAndCacheMessageType(TSharedPtr<FUdpReassembledMessage, ESPMode::ThreadSafe>& ReassembledMessage);
+	bool LookupAndCacheMessageType(TSharedPtr<FUdpReassembledMessage>& ReassembledMessage);
 
 	/**
 	 * Removes the specified node from the list of known remote endpoints.
@@ -1023,6 +1057,9 @@ private:
 	/** Checks all known nodes to see if any have both segmenters and send windows ready to transmit. */
 	bool MoreToSend();
 
+	/** Increments and returns the running message ID, bounding it to be non-negative. */
+	int32 AssignNextMessageId();
+
 	/** Consume an Outbound Message for processing. */
 	void ConsumeOneOutboundMessage(const FOutboundMessage& OutboundMessage);
 
@@ -1059,6 +1096,14 @@ private:
 	/** Indicate if we added new nodes during this frame. */
 	bool bAddedNewKnownNodes = false;
 
+#if UDPMESSAGINGTRACE_ENABLED
+	/** FNodeInfo trace ID. 0 == invalid. */
+	static std::atomic<uint16> NextNodeTraceId;
+
+	/** Holds the local node trace ID. */
+	uint16 LocalNodeTraceId;
+#endif
+
 	/** Holds the local node identifier. */
 	FGuid LocalNodeId;
 
diff --git a/Engine/Plugins/Messaging/UdpMessaging/Source/UdpMessaging/Private/Transport/UdpMessageProcessorLogger.cpp b/Engine/Plugins/Messaging/UdpMessaging/Source/UdpMessaging/Private/Transport/UdpMessageProcessorLogger.cpp
--- a/Engine/Plugins/Messaging/UdpMessaging/Source/UdpMessaging/Private/Transport/UdpMessageProcessorLogger.cpp
+++ b/Engine/Plugins/Messaging/UdpMessaging/Source/UdpMessaging/Private/Transport/UdpMessageProcessorLogger.cpp
@@ -431,7 +431,7 @@ namespace UE::UdpMessaging::Private
 		TStringBuilder<512> Builder;
 		Builder.Append(Stat.Endpoint.ToString());
 		Builder.Append(TEXT(","));
-		Builder.Appendf(TEXT("%.3f"), Stat.LastSegmentReceivedTime.RawValue);
+		Builder.Appendf(TEXT("%s"), *Stat.LastSegmentReceivedTime.EstimatedLocalDateTime().ToString());
 		Builder.Append(TEXT(","));
 		Builder.Append(Stat.NodeId.ToString());
 		Builder.Append(TEXT(","));
diff --git a/Engine/Plugins/Messaging/UdpMessaging/Source/UdpMessaging/Private/Transport/UdpMessageSegmenter.cpp b/Engine/Plugins/Messaging/UdpMessaging/Source/UdpMessaging/Private/Transport/UdpMessageSegmenter.cpp
--- a/Engine/Plugins/Messaging/UdpMessaging/Source/UdpMessaging/Private/Transport/UdpMessageSegmenter.cpp
+++ b/Engine/Plugins/Messaging/UdpMessaging/Source/UdpMessaging/Private/Transport/UdpMessageSegmenter.cpp
@@ -108,29 +108,33 @@ bool FUdpMessageSegmenter::GetPendingSegment(uint32 InSegment, TArray<uint8>& Ou
 }
 
 
-void FUdpMessageSegmenter::Initialize()
+bool FUdpMessageSegmenter::TryInitialize()
 {
-	if (MessageReader != nullptr)
+	if (!ensureMsgf(!IsInitialized(), TEXT("Tried to initialize an already initialized segmenter")))
 	{
-		return;
+		return false;
 	}
 
-	if (SerializedMessage->GetState() == EUdpSerializedMessageState::Complete)
+	if (SerializedMessage->GetState() != EUdpSerializedMessageState::Complete)
 	{
-		MessageReader = SerializedMessage->CreateReader();
-		PendingSendSegmentsCount = (MessageReader->TotalSize() + GetSegmentPayloadSize() - 1) / GetSegmentPayloadSize();
-		PendingSendSegments.Init(true, PendingSendSegmentsCount);
-		if (EnumHasAnyFlags(GetMessageFlags(), EMessageFlags::Reliable))
-		{
-			AcknowledgeSegments.Init(false, PendingSendSegmentsCount);
-		}
-		else
-		{
-			// Acks for unreliable messages are always zero
-			AcknowledgeSegments.Init(false, 0);
-		}		
-		AcknowledgeSegmentsCount = 0;
+		return false;
 	}
+
+	MessageReader = SerializedMessage->CreateReader();
+	PendingSendSegmentsCount = (MessageReader->TotalSize() + GetSegmentPayloadSize() - 1) / GetSegmentPayloadSize();
+	PendingSendSegments.Init(true, PendingSendSegmentsCount);
+	if (EnumHasAnyFlags(GetMessageFlags(), EMessageFlags::Reliable))
+	{
+		AcknowledgeSegments.Init(false, PendingSendSegmentsCount);
+	}
+	else
+	{
+		// Acks for unreliable messages are always zero
+		AcknowledgeSegments.Init(false, 0);
+	}
+	AcknowledgeSegmentsCount = 0;
+
+	return true;
 }
 
 
@@ -272,3 +276,43 @@ bool FUdpMessageSegmenter::WasRetransmitted(uint32 SegmentId) const
 
 	return RetransmittedSegments.Contains(SegmentId);
 }
+
+#if UDPMESSAGINGTRACE_ENABLED
+void FUdpMessageSegmenter::EmitPostInitializeTraceEvents(uint16 InSenderTraceId, uint16 InRecipientTraceId, int32 InMessageId)
+{
+	const TVariant<FName, FTopLevelAssetPath> TypeInfo(
+		TInPlaceType<FTopLevelAssetPath>(), SerializedMessage->GetTraceMetadata().TypeInfo);
+
+	TRACE_UDPMESSAGING_MESSAGE_TYPEINFO(
+		InSenderTraceId,
+		InRecipientTraceId,
+		InMessageId,
+		TypeInfo
+	);
+
+	TRACE_UDPMESSAGING_MESSAGE_LIFECYCLE(
+		SerializedMessage->GetTraceMetadata().SerializationStartTime.Cycles,
+		InSenderTraceId,
+		InRecipientTraceId,
+		InMessageId,
+		UE::UdpMessaging::Trace::EUdpMessageLifecycleEvent::BeginSerialize
+	);
+
+	TRACE_UDPMESSAGING_MESSAGE_LIFECYCLE(
+		SerializedMessage->GetTraceMetadata().SerializationEndTime.Cycles,
+		InSenderTraceId,
+		InRecipientTraceId,
+		InMessageId,
+		UE::UdpMessaging::Trace::EUdpMessageLifecycleEvent::EndSerialize
+	);
+
+	TRACE_UDPMESSAGING_MESSAGE_SUMMARY(
+		InSenderTraceId,
+		InRecipientTraceId,
+		InMessageId,
+		GetMessageSize(),
+		GetSegmentCount(),
+		GetMessageFlags()
+	);
+}
+#endif // #if UDPMESSAGINGTRACE_ENABLED
diff --git a/Engine/Plugins/Messaging/UdpMessaging/Source/UdpMessaging/Private/Transport/UdpMessageSegmenter.h b/Engine/Plugins/Messaging/UdpMessaging/Source/UdpMessaging/Private/Transport/UdpMessageSegmenter.h
--- a/Engine/Plugins/Messaging/UdpMessaging/Source/UdpMessaging/Private/Transport/UdpMessageSegmenter.h
+++ b/Engine/Plugins/Messaging/UdpMessaging/Source/UdpMessaging/Private/Transport/UdpMessageSegmenter.h
@@ -6,6 +6,7 @@
 #include "Shared/UdpMessageSegment.h"
 #include "Shared/UdpMessagingTime.h"
 #include "Templates/SharedPointer.h"
+#include "Trace/UdpMessagingTrace.h"
 
 
 // IMessageContext forward declaration
@@ -100,7 +101,7 @@ public:
 	/**
 	 * Gets the number of segments that have been acknowledged. Always zero for unreliable messages
 	 *
-	 * @return Number of pending segments.
+	 * @return Number of acknowledged segments.
 	 */
 	uint32 GetAcknowledgedSegmentsCount() const
 	{
@@ -148,8 +149,12 @@ public:
 		return SegmentSize - GetDataSegmentHeaderSize();
 	}
 
-	/** Initializes the segmenter. */
-	void Initialize();
+	/**
+	 * Attempt to initialize an uninitialized segmenter.
+	 *
+	 * @return true if the segmenter was uninitialized and serialization is complete, false otherwise.
+	 */
+	bool TryInitialize();
 
 	/**
 	 * Checks whether all segments have been sent.
@@ -247,12 +252,17 @@ public:
 	 * Checks whether the serialization of the message to segment and send completed. The serialization may be performing asynchronously and can may succeed or fail. The message cannot be segmented
 	 * and sent until it was serialized sucessfully.
 	 * @return true if the message is serialization is done, false if it still serializing.
-	 * @see Initialize()
+	 * @see TryInitialize()
 	 * @see IsInvalid()
 	 * @see IsInitialized()
 	 */
 	bool IsMessageSerializationDone() const;
 
+#if UDPMESSAGINGTRACE_ENABLED
+	/** Once serialization completes, the processor passes in context we use to record facts about this message. */
+	void EmitPostInitializeTraceEvents(uint16 InSenderTraceId, uint16 InRecipientTraceId, int32 InMessageId);
+#endif
+
 private:
 	/** Used to seek and output portions of SerializedMessage. */
 	FArchive* MessageReader = nullptr;
diff --git a/Engine/Plugins/Messaging/UdpMessaging/Source/UdpMessaging/Private/Transport/UdpMessageTransport.cpp b/Engine/Plugins/Messaging/UdpMessaging/Source/UdpMessaging/Private/Transport/UdpMessageTransport.cpp
--- a/Engine/Plugins/Messaging/UdpMessaging/Source/UdpMessaging/Private/Transport/UdpMessageTransport.cpp
+++ b/Engine/Plugins/Messaging/UdpMessaging/Source/UdpMessaging/Private/Transport/UdpMessageTransport.cpp
@@ -12,17 +12,15 @@
 #include "Shared/UdpMessagingSettings.h"
 #include "Tasks/Task.h"
 #include "Tasks/TaskConcurrencyLimiter.h"
+#include "Trace/UdpMessagingTrace.h"
 #include "Transport/UdpReassembledMessage.h"
 #include "Transport/UdpDeserializedMessage.h"
 #include "Transport/UdpMessageProcessor.h"
-#include "UdpMessageProcessor.h"
 
 #if PLATFORM_DESKTOP
 #include "UdpMessageProcessorLogger.h"
 #endif
 
-/* FUdpMessageTransport structors
- *****************************************************************************/
 
 TAutoConsoleVariable<int32> CVarMaxRetriesForBadEndpoint(
 	TEXT("MessageBus.UDP.MaxRetriesForBadEndpoint"),
@@ -469,6 +467,15 @@ void FUdpMessageTransport::HandleProcessorMessageReassembled(FUdpReassembledMess
 	auto DeserializeReassembledMessage = [DeserializedMessage = MoveTemp(DeserializedMessage), TransportHandler = this->TransportHandler, ReassembledMessagePtr = MoveTemp(ReassembledMessagePtr), NodeId](uint32 Slot)
 	{
 		const FUdpReassembledMessage& ReassembledMessage = *ReassembledMessagePtr.Get();
+
+		TRACE_UDPMESSAGING_MESSAGE_LIFECYCLE(
+			FPlatformTime::Cycles64(),
+			ReassembledMessage.GetTraceMetadata().SenderShortId,
+			ReassembledMessage.GetTraceMetadata().RecipientShortId,
+			ReassembledMessage.GetTraceMetadata().MessageId,
+			UE::UdpMessaging::Trace::EUdpMessageLifecycleEvent::BeginDeserialize
+		);
+
 		if (DeserializedMessage->Deserialize(ReassembledMessage))
 		{
 			TransportHandler->ReceiveTransportMessage(DeserializedMessage, NodeId);
@@ -479,14 +486,22 @@ void FUdpMessageTransport::HandleProcessorMessageReassembled(FUdpReassembledMess
 			FString TypeInfoStr = TypeInfo.IsValid() ? TypeInfo->GetName() : TEXT("UNKNOWN TYPE");
 			UE_LOG(LogUdpMessaging, Verbose, TEXT("Failed to deserialize message of type %s from %s"), *TypeInfoStr, *NodeId.ToString());
 		}
+
 		// We always mark the message as delivered even if we can't deserialize it. There may be some messages that can't be deserialized because of missing type information.
 		ReassembledMessagePtr->MarkDelivered();
+
+		TRACE_UDPMESSAGING_MESSAGE_LIFECYCLE(
+			FPlatformTime::Cycles64(),
+			ReassembledMessage.GetTraceMetadata().SenderShortId,
+			ReassembledMessage.GetTraceMetadata().RecipientShortId,
+			ReassembledMessage.GetTraceMetadata().MessageId,
+			UE::UdpMessaging::Trace::EUdpMessageLifecycleEvent::Delivered
+		);
 	};
 
 	if (DeserializationTasks)
 	{
 		// The user has asked for async deserialization.  Dispatch the task.
-		// 
 		DeserializationTasks->Push(UE_SOURCE_LOCATION, MoveTemp(DeserializeReassembledMessage));
 	}
 	else
diff --git a/Engine/Plugins/Messaging/UdpMessaging/Source/UdpMessaging/Private/Transport/UdpReassembledMessage.h b/Engine/Plugins/Messaging/UdpMessaging/Source/UdpMessaging/Private/Transport/UdpReassembledMessage.h
--- a/Engine/Plugins/Messaging/UdpMessaging/Source/UdpMessaging/Private/Transport/UdpReassembledMessage.h
+++ b/Engine/Plugins/Messaging/UdpMessaging/Source/UdpMessaging/Private/Transport/UdpReassembledMessage.h
@@ -5,6 +5,7 @@
 #include "IMessageContext.h"
 #include "Interfaces/IPv4/IPv4Endpoint.h"
 #include "Shared/UdpMessagingTime.h"
+#include "Trace/UdpMessagingTrace.h"
 #include "UdpMessagingPrivate.h"
 
 // IMessageContext forward declaration
@@ -359,6 +360,26 @@ public:
 	{
 		TypeInfo = MoveTemp(InTypeInfo);
 	}
+
+#if UDPMESSAGINGTRACE_ENABLED
+	struct FTraceMetadata
+	{
+		uint16 SenderShortId = 0;
+		uint16 RecipientShortId = 0;
+		int32 MessageId = 0;
+	};
+
+	const FTraceMetadata& GetTraceMetadata() const
+	{
+		return TraceMetadata;
+	}
+
+	FTraceMetadata& GetTraceMetadata()
+	{
+		return TraceMetadata;
+	}
+#endif
+
 private:
 	/** Holds the message's type information. */
 	TWeakObjectPtr<UScriptStruct> TypeInfo;
@@ -404,4 +425,8 @@ private:
 
 	/** Indicates whether or not this message is malformed. */
 	bool bIsMalformed = false;
+
+#if UDPMESSAGINGTRACE_ENABLED
+	FTraceMetadata TraceMetadata;
+#endif
 };
diff --git a/Engine/Plugins/Messaging/UdpMessaging/Source/UdpMessaging/Private/Transport/UdpSerializeMessageTask.cpp b/Engine/Plugins/Messaging/UdpMessaging/Source/UdpMessaging/Private/Transport/UdpSerializeMessageTask.cpp
--- a/Engine/Plugins/Messaging/UdpMessaging/Source/UdpMessaging/Private/Transport/UdpSerializeMessageTask.cpp
+++ b/Engine/Plugins/Messaging/UdpMessaging/Source/UdpMessaging/Private/Transport/UdpSerializeMessageTask.cpp
@@ -11,7 +11,7 @@
 #include "Transport/UdpSerializedMessage.h"
 #include "Shared/UdpMessagingSettings.h"
 #include "UdpMessageSegmenter.h"
-#include "UdpMessagingTracing.h"
+#include "Trace/UdpMessagingTrace.h"
 
 namespace UdpSerializeMessageTaskDetails
 {
@@ -146,6 +146,13 @@ void SerializeMessageV16_18(FArchive& Archive, const TSharedRef<IMessageContext,
 void FUdpSerializeMessageTask::DoTask(ENamedThreads::Type CurrentThread, const FGraphEventRef& MyCompletionGraphEvent)
 {
 	SCOPED_MESSAGING_TRACE(FUdpSerializeMessageTask_DoTask);
+
+#if UDPMESSAGINGTRACE_ENABLED
+	FUdpSerializedMessage::FTraceMetadata& TraceMetadata = SerializedMessage->GetTraceMetadata();
+	TraceMetadata.TypeInfo = MessageContext->GetMessageTypePathName();
+	TraceMetadata.SerializationStartTime = FUdpMessagingTime::Now();
+#endif
+
 	if (MessageContext->IsValid())
 	{
 		// Note that some complex values are serialized manually here, so that we can ensure
@@ -245,6 +252,10 @@ void FUdpSerializeMessageTask::DoTask(ENamedThreads::Type CurrentThread, const F
 		SerializedMessage->UpdateState(EUdpSerializedMessageState::Invalid);
 	}
 
+#if UDPMESSAGINGTRACE_ENABLED
+	TraceMetadata.SerializationEndTime = FUdpMessagingTime::Now();
+#endif
+
 	UE_LOG(LogUdpMessaging, Verbose, TEXT("Serialized %s from %s to %" INT64_FMT " bytes"), *MessageContext->GetMessageTypePathName().ToString(), *MessageContext->GetSender().ToString(), SerializedMessage->TotalSize());
 
 	// signal task completion
diff --git a/Engine/Plugins/Messaging/UdpMessaging/Source/UdpMessaging/Private/Transport/UdpSerializedMessage.h b/Engine/Plugins/Messaging/UdpMessaging/Source/UdpMessaging/Private/Transport/UdpSerializedMessage.h
--- a/Engine/Plugins/Messaging/UdpMessaging/Source/UdpMessaging/Private/Transport/UdpSerializedMessage.h
+++ b/Engine/Plugins/Messaging/UdpMessaging/Source/UdpMessaging/Private/Transport/UdpSerializedMessage.h
@@ -5,13 +5,15 @@
 #include "Serialization/MemoryWriter.h"
 #include "Serialization/MemoryReader.h"
 #include "IMessageContext.h"
+#include "Shared/UdpMessagingTime.h"
+#include "Trace/UdpMessagingTrace.h"
 
 #include <atomic>
 
 enum class EUdpMessageFormat : uint8;
 
 /**
- * Enumerates possibly states of a serialized message.
+ * Enumerates possible states of a serialized message.
  *
  * @see FUdpSerializedMessage
  */
@@ -115,6 +117,25 @@ public:
 		return Format;
 	}
 
+#if UDPMESSAGINGTRACE_ENABLED
+	struct FTraceMetadata
+	{
+		FTopLevelAssetPath TypeInfo;
+		FUdpMessagingTime SerializationStartTime;
+		FUdpMessagingTime SerializationEndTime;
+	};
+
+	const FTraceMetadata& GetTraceMetadata() const
+	{
+		return TraceMetadata;
+	}
+
+	FTraceMetadata& GetTraceMetadata()
+	{
+		return TraceMetadata;
+	}
+#endif
+
 private:
 	/** Holds the serialized data. */
 	TArray<uint8> DataArray;
@@ -130,4 +151,8 @@ private:
 
 	/** Holds the Protocol Version the message will be serialized in. */
 	uint8 ProtocolVersion;
+
+#if UDPMESSAGINGTRACE_ENABLED
+	FTraceMetadata TraceMetadata;
+#endif
 };
diff --git a/Engine/Plugins/Messaging/UdpMessaging/Source/UdpMessaging/Private/UdpMessagingTracing.cpp b/Engine/Plugins/Messaging/UdpMessaging/Source/UdpMessaging/Private/UdpMessagingTracing.cpp
deleted file mode 100644
diff --git a/Engine/Plugins/Messaging/UdpMessaging/Source/UdpMessaging/Public/UdpMessagingTracing.h b/Engine/Plugins/Messaging/UdpMessaging/Source/UdpMessaging/Public/UdpMessagingTracing.h
--- a/Engine/Plugins/Messaging/UdpMessaging/Source/UdpMessaging/Public/UdpMessagingTracing.h
+++ b/Engine/Plugins/Messaging/UdpMessaging/Source/UdpMessaging/Public/UdpMessagingTracing.h
@@ -2,13 +2,8 @@
 
 #pragma once
 
-#include "ProfilingDebugging/CpuProfilerTrace.h"
+#include "Misc/CoreMiscDefines.h"
 
-namespace UE::Trace { class FChannel; }
+UE_DEPRECATED_HEADER(5.8, "SCOPED_MESSAGING_TRACE and MessagingChannel are now exported by MessagingTrace.h in the Messaging module.")
 
-#if UE_TRACE_ENABLED && !UE_BUILD_SHIPPING
-UE_TRACE_CHANNEL_EXTERN(MessagingChannel, UDPMESSAGING_API);
-#define SCOPED_MESSAGING_TRACE(TraceName) TRACE_CPUPROFILER_EVENT_SCOPE_ON_CHANNEL(TraceName, MessagingChannel)
-#else
-#define SCOPED_MESSAGING_TRACE(TraceName)
-#endif
+#include "MessagingTrace.h"
diff --git a/Engine/Programs/UnrealInsights/Config/DefaultEngine.ini b/Engine/Programs/UnrealInsights/Config/DefaultEngine.ini
--- a/Engine/Programs/UnrealInsights/Config/DefaultEngine.ini
+++ b/Engine/Programs/UnrealInsights/Config/DefaultEngine.ini
@@ -11,6 +11,7 @@
 +ProgramEnabledPlugins=XcodeSourceCodeAccess
 +ProgramEnabledPlugins=MassInsights
 +ProgramEnabledPlugins=ChaosInsights
++ProgramEnabledPlugins=MessagingInsights
 
 [/Script/MacTargetPlatform.XcodeProjectSettings]
 TemplateMacPlist=(FilePath="/Engine/Source/Programs/UnrealInsights/Resources/Mac/Info.plist")
diff --git a/Engine/Source/Runtime/Messaging/Messaging.Build.cs b/Engine/Source/Runtime/Messaging/Messaging.Build.cs
--- a/Engine/Source/Runtime/Messaging/Messaging.Build.cs
+++ b/Engine/Source/Runtime/Messaging/Messaging.Build.cs
@@ -11,6 +11,7 @@ namespace UnrealBuildTool.Rules
 				{
 					"Core",
 					"CoreUObject",
+					"TraceLog",
 				});
 
 			PrivateIncludePaths.AddRange(
diff --git a/Engine/Source/Runtime/Messaging/Private/MessagingModule.cpp b/Engine/Source/Runtime/Messaging/Private/MessagingModule.cpp
--- a/Engine/Source/Runtime/Messaging/Private/MessagingModule.cpp
+++ b/Engine/Source/Runtime/Messaging/Private/MessagingModule.cpp
@@ -11,6 +11,7 @@
 #include "Bridge/MessageBridge.h"
 #include "IMessagingModule.h"
 #include "INetworkMessagingExtension.h"
+#include "MessagingTrace.h"
 
 
 #ifndef PLATFORM_SUPPORTS_MESSAGEBUS
@@ -19,6 +20,11 @@
 
 DEFINE_LOG_CATEGORY(LogMessaging);
 
+#if MESSAGINGTRACE_ENABLED
+UE_TRACE_CHANNEL_DEFINE(MessagingChannel);
+#endif
+
+
 /**
  * Implements the Messaging module.
  */
diff --git a/Engine/Source/Runtime/Messaging/Public/MessagingTrace.h b/Engine/Source/Runtime/Messaging/Public/MessagingTrace.h
new file mode 100644
--- /dev/null
+++ b/Engine/Source/Runtime/Messaging/Public/MessagingTrace.h
@@ -0,0 +1,29 @@
+// Copyright Epic Games, Inc. All Rights Reserved.
+
+#pragma once
+
+#include "Trace/Config.h"
+
+
+#if !defined(MESSAGINGTRACE_ENABLED)
+#	if UE_TRACE_ENABLED && !UE_BUILD_SHIPPING
+#		define MESSAGINGTRACE_ENABLED 1
+#	else
+#		define MESSAGINGTRACE_ENABLED 0
+#	endif
+#endif
+
+
+#if MESSAGINGTRACE_ENABLED
+
+#include "Trace/Trace.h"
+
+UE_TRACE_CHANNEL_EXTERN(MessagingChannel, MESSAGING_API);
+
+#define SCOPED_MESSAGING_TRACE(TraceName) TRACE_CPUPROFILER_EVENT_SCOPE_ON_CHANNEL(TraceName, MessagingChannel)
+
+#else // #if MESSAGINGTRACE_ENABLED
+
+#define SCOPED_MESSAGING_TRACE(...)
+
+#endif // #if MESSAGINGTRACE_ENABLED .. #else
