UdpMessaging: Implement token bucket send pacing.

The MaxSendRate setting, unused in recent versions of the engine, should now be respected.

Previous bottlenecks artificially limited the outgoing bandwidth to less than the default rate, so the limit not being enforced was inconsequential. This became more important to fix because recent optimizations have made it possible for UdpMessaging to consume ~3x the previous amount of bandwidth (~70 MiB/s -> ~200 MiB/s).

The token bucket is implemented in a lock-free fashion to enable it to be used to jointly govern the egress of multiple sockets in the future.

#Synced-CL 49810677

diff --git a/Engine/Plugins/Messaging/UdpMessaging/Source/UdpMessaging/Private/Shared/SocketSender.cpp b/Engine/Plugins/Messaging/UdpMessaging/Source/UdpMessaging/Private/Shared/SocketSender.cpp
--- a/Engine/Plugins/Messaging/UdpMessaging/Source/UdpMessaging/Private/Shared/SocketSender.cpp
+++ b/Engine/Plugins/Messaging/UdpMessaging/Source/UdpMessaging/Private/Shared/SocketSender.cpp
@@ -22,6 +22,7 @@ FSocketSender::FSocketSender(FSocket* InSocket, const TCHAR* InDescription, cons
 	, bStopping(false)
 	, WaitTime(InOptions.WaitTime)
 	, Description(InDescription)
+	, TokenBucket(InOptions.BytesPerSec, InOptions.MaxBurstBytes)
 {
 	check(Socket != nullptr);
 	check(Socket->GetSocketType() == SOCKTYPE_Datagram);
@@ -125,11 +126,29 @@ FSocketSender::EUpdateResult FSocketSender::Update(const FTimespan& InSocketWait
 
 		SCOPED_MESSAGING_TRACE(UdpMessaging_FSocketSender_Update);
 
-		TMpscQueue<FPacket>& Queue = !PrioritySendQueue.IsEmpty() ? PrioritySendQueue : SendQueue;
+		const bool bIsPriority = !PrioritySendQueue.IsEmpty();
+		TMpscQueue<FPacket>& Queue = bIsPriority ? PrioritySendQueue : SendQueue;
 
 		const FPacket* Packet = Queue.Peek();
 		const FIPv4Endpoint Recipient = Packet->Recipient;
 		const int32 PacketNumBytes = Packet->Data->Num();
+
+		const int32 TokenBytes = PacketNumBytes + PacketHeaderBytes;
+		const bool bConsumedTokens = TokenBucket.TryConsume(TokenBytes);
+		if (!bConsumedTokens && !bIsPriority)
+		{
+			if (!ensure(TokenBytes <= TokenBucket.Max()))
+			{
+				UE_LOGFMT(LogUdpMessaging, Error,
+					"Packet size {TokenBytes} exceeds max burst size {MaxBytes}",
+					TokenBytes, TokenBucket.Max());
+
+				return EUpdateResult::Fatal;
+			}
+
+			return EUpdateResult::Retry;
+		}
+
 		int32 SentNumBytes = 0;
 		const bool bSendResult =
 			Socket->SendTo(Packet->Data->GetData(), PacketNumBytes, SentNumBytes, *Recipient.ToInternetAddr());
@@ -138,6 +157,12 @@ FSocketSender::EUpdateResult FSocketSender::Update(const FTimespan& InSocketWait
 
 		if (!bSendResult)
 		{
+			if (bConsumedTokens)
+			{
+				// Refund tokens that never actually egressed.
+				TokenBucket.Credit(TokenBytes);
+			}
+
 			SocketError = SocketSubsystem->GetLastErrorCode();
 			if (SocketError == SE_EWOULDBLOCK || SocketError == SE_ENOBUFS || SocketError == SE_EINTR)
 			{
diff --git a/Engine/Plugins/Messaging/UdpMessaging/Source/UdpMessaging/Private/Shared/SocketSender.h b/Engine/Plugins/Messaging/UdpMessaging/Source/UdpMessaging/Private/Shared/SocketSender.h
--- a/Engine/Plugins/Messaging/UdpMessaging/Source/UdpMessaging/Private/Shared/SocketSender.h
+++ b/Engine/Plugins/Messaging/UdpMessaging/Source/UdpMessaging/Private/Shared/SocketSender.h
@@ -7,6 +7,7 @@
 #include "HAL/Runnable.h"
 #include "Interfaces/IPv4/IPv4Endpoint.h"
 #include "Misc/SingleThreadRunnable.h"
+#include "Shared/TokenBucket.h"
 #include "Templates/SharedPointer.h"
 
 
@@ -66,6 +67,12 @@ public:
 		/** Datagram payload size for UDP segmentation offload; 0 disables USO. Ignored if unsupported. */
 		uint16 SegmentOffloadSize = 0;
 
+		/** Token bucket: steady state send bytes per second / replenishment rate. */
+		int64 BytesPerSec = 100 * 1024 * 1024;
+
+		/** Token bucket: max bucket reserve capacity. */
+		int64 MaxBurstBytes = 10 * 1024 * 1024;
+
 		FOptions()
 		{
 			// Without this (or with `= default`), clang gives the following:
@@ -172,6 +179,9 @@ private:
 
 	/** Datagram payload size for segmentation offload. 0 means disabled or unsupported. */
 	uint16 SegmentOffloadSize;
+
+	/** Token bucket for egress shaping. */
+	FLockFreeTokenBucket TokenBucket;
 };
 
 } // namespace UE::UdpMessaging
diff --git a/Engine/Plugins/Messaging/UdpMessaging/Source/UdpMessaging/Private/Shared/TokenBucket.h b/Engine/Plugins/Messaging/UdpMessaging/Source/UdpMessaging/Private/Shared/TokenBucket.h
new file mode 100644
--- /dev/null
+++ b/Engine/Plugins/Messaging/UdpMessaging/Source/UdpMessaging/Private/Shared/TokenBucket.h
@@ -0,0 +1,80 @@
+// Copyright Epic Games, Inc. All Rights Reserved.
+
+#pragma once
+
+#include "HAL/Platform.h"
+#include "HAL/PlatformTime.h"
+
+#include <atomic>
+
+
+namespace UE::UdpMessaging
+{
+
+class FLockFreeTokenBucket
+{
+public:
+	FLockFreeTokenBucket(int64 InTokensPerSec, int64 InMaxBurstTokens)
+		: TokensPerSec(InTokensPerSec)
+		, MaxBurstTokens(InMaxBurstTokens)
+	{
+		check(TokensPerSec > 0);
+		check(MaxBurstTokens > 0);
+
+		FullAtTimeCycles = FPlatformTime::Cycles64();
+	}
+
+	int64 Max() const
+	{
+		return MaxBurstTokens;
+	}
+
+	bool TryConsume(int64 InConsumeTokens)
+	{
+		check(InConsumeTokens > 0);
+
+		if (UNLIKELY(!ensure(InConsumeTokens <= MaxBurstTokens)))
+		{
+			return false;
+		}
+
+		const double SecondsPerCycle64 = FPlatformTime::GetSecondsPerCycle64();
+		int64 Expected = FullAtTimeCycles.load(std::memory_order_relaxed);
+		int64 Desired;
+		do 
+		{
+			const int64 NowCycles = FPlatformTime::Cycles64();
+			const double SecToFull = (Expected - NowCycles) * SecondsPerCycle64;
+			const int64 TokensAvailable = FMath::Clamp(MaxBurstTokens - SecToFull * TokensPerSec, 0ll, MaxBurstTokens);
+			if (InConsumeTokens > TokensAvailable)
+			{
+				return false;
+			}
+
+			const double ConsumedSec = InConsumeTokens / static_cast<double>(TokensPerSec);
+			const int64 ConsumedCycles = FMath::CeilToInt64(ConsumedSec / SecondsPerCycle64);
+			Desired = FMath::Max(Expected, NowCycles) + ConsumedCycles;
+		} while (!FullAtTimeCycles.compare_exchange_weak(Expected, Desired));
+
+		return true;
+	}
+
+	void Credit(int64 InCreditTokens)
+	{
+		check(InCreditTokens > 0);
+
+		const double SecondsPerCycle64 = FPlatformTime::GetSecondsPerCycle64();
+		const double CreditSec = InCreditTokens / static_cast<double>(TokensPerSec);
+		const int64 CreditCycles = FMath::CeilToInt64(CreditSec / SecondsPerCycle64);
+		FullAtTimeCycles.fetch_sub(CreditCycles);
+	}
+
+private:
+	const int64 TokensPerSec;
+	const int64 MaxBurstTokens;
+
+	/** A single atomic value that represents the current state. If <= now, the bucket is full. */
+	std::atomic<int64> FullAtTimeCycles;
+};
+
+} // namespace UE::UdpMessaging
diff --git a/Engine/Plugins/Messaging/UdpMessaging/Source/UdpMessaging/Private/Transport/UdpMessageProcessor.cpp b/Engine/Plugins/Messaging/UdpMessaging/Source/UdpMessaging/Private/Transport/UdpMessageProcessor.cpp
--- a/Engine/Plugins/Messaging/UdpMessaging/Source/UdpMessaging/Private/Transport/UdpMessageProcessor.cpp
+++ b/Engine/Plugins/Messaging/UdpMessaging/Source/UdpMessaging/Private/Transport/UdpMessageProcessor.cpp
@@ -354,11 +354,16 @@ bool FUdpMessageProcessor::Init()
 		const UUdpMessagingSettings* Settings = GetDefault<UUdpMessagingSettings>();
 
 		UE::UdpMessaging::FSocketSender::FOptions SenderOpts;
+
 		SenderOpts.SegmentOffloadSize =
 			Settings->bSegmentationOffload
 			? (Settings->GetMaxPacketSize() - UE::UdpMessaging::PacketHeaderBytes)
 			: 0;
 
+		const uint64 OneGbitPerSecondInBytes = 125000000;
+		SenderOpts.BytesPerSec = GetDefault<UUdpMessagingSettings>()->GetMaxSendRate() * OneGbitPerSecondInBytes;
+		SenderOpts.MaxBurstBytes = SenderOpts.BytesPerSec * 0.1; // Absorbs 100ms jitter
+
 		SocketSender = MakeUnique<UE::UdpMessaging::FSocketSender>(Socket, TEXT("FUdpMessageProcessor.Sender"), SenderOpts);
 
 		Beacon = MakeUnique<FUdpMessageBeacon>(SocketSender.Get(), LocalNodeId, MulticastEndpoint);
@@ -1155,12 +1160,6 @@ void FUdpMessageProcessor::RemoveKnownNode(const FGuid& NodeId)
 	KnownNodes.Remove(NodeId);
 }
 
-int32 GetMaxSendRate()
-{
-	const uint32 OneGbitPerSecondInBytes = 125000000;
-	return static_cast< uint32 > ( GetDefault<UUdpMessagingSettings>()->MaxSendRate * OneGbitPerSecondInBytes );
-}
-
 bool HasTimedOut(const FUdpMessagingTime& ReceivedTime, const FUdpMessagingTime& CurrentTime)
 {
 	const float TimeoutInSeconds = GetDefault<UUdpMessagingSettings>()->ConnectionTimeoutPeriod;
@@ -1774,10 +1773,6 @@ bool FUdpMessageProcessor::UpdateReassemblers(FNodeInfo& NodeInfo)
 	}
 
 	const FUdpMessagingTime StaleLastSegmentTime = FUdpMessagingTime::Now() - StaleReassemblyInterval;
-	const uint32 MaxSendRate = GetMaxSendRate();
-	double DeltaSeconds = DeltaTime.GetTotalSeconds();
-	const int32 MaxSendRateDelta = MaxSendRate * DeltaSeconds;
-	int32 MaxNodeByteSend = (MaxSendRateDelta) / KnownNodes.Num();
 
 	for (TMap<int32, TSharedPtr<FUdpReassembledMessage, ESPMode::ThreadSafe>>::TIterator It(NodeInfo.ReassembledMessages); It; ++It)
 	{
diff --git a/Engine/Plugins/Messaging/UdpMessaging/Source/UdpMessaging/Public/Shared/UdpMessagingSettings.h b/Engine/Plugins/Messaging/UdpMessaging/Source/UdpMessaging/Public/Shared/UdpMessagingSettings.h
--- a/Engine/Plugins/Messaging/UdpMessaging/Source/UdpMessaging/Public/Shared/UdpMessagingSettings.h
+++ b/Engine/Plugins/Messaging/UdpMessaging/Source/UdpMessaging/Public/Shared/UdpMessagingSettings.h
@@ -70,11 +70,19 @@ public:
 	 * by your network card.
 	 *
 	 * Adjust this value higher or lower depending on your network capacity.
-	 *
 	 */
-	UPROPERTY(config, EditAnywhere, Category=Transport, AdvancedDisplay)
+	UPROPERTY(config, EditAnywhere, Category=Transport, meta = (ClampMin="0.01", ClampMax="100.0"), AdvancedDisplay)
 	float MaxSendRate = 1.0f;
 
+	static constexpr float MaxSendRate_Min = 0.01;
+	static constexpr float MaxSendRate_Max = 100.0;
+
+	/** Accessor that ensures MaxSendRate is clamped to the expected bounds. */
+	float GetMaxSendRate() const
+	{
+		return FMath::Clamp<float>(MaxSendRate, MaxSendRate_Min, MaxSendRate_Max);
+	}
+
 	/** The number of consecutive attempt the auto repair routine will try to repair. */
 	UPROPERTY(config, EditAnywhere, Category = Transport, AdvancedDisplay)
 	uint32 AutoRepairAttemptLimit = 10;
